Copying //UE4/Dev-Editor to Dev-Main (//UE4/Dev-Main) for 4.24

diff --git a/Engine/Plugins/Experimental/MeshModelingToolset/Source/ModelingComponents/Public/Snapping/PointPlanarSnapSolver.h b/Engine/Plugins/Experimental/MeshModelingToolset/Source/ModelingComponents/Public/Snapping/PointPlanarSnapSolver.h
new file mode 100644
--- /dev/null
+++ b/Engine/Plugins/Experimental/MeshModelingToolset/Source/ModelingComponents/Public/Snapping/PointPlanarSnapSolver.h
@@ -0,0 +1,57 @@
+// Copyright 1998-2019 Epic Games, Inc. All Rights Reserved.
+
+#pragma once
+
+#include "Snapping/BasePositionSnapSolver3.h"
+
+/**
+ * FPointPlanarSnapSolver solves for a Point snap location on a plane, based
+ * on an input Point and a set of target points and lines in the plane.
+ *
+ * This implementation has the notion of a "history" of previous points,
+ * from which line and distance constraints can be inferred.
+ * This is useful for snapping in 2D polygon drawing.
+ *
+ * See FBasePositionSnapSolver3 for details on how to set up the snap problem
+ * and get results.
+ */
+class MODELINGCOMPONENTS_API FPointPlanarSnapSolver : public FBasePositionSnapSolver3
+{
+public:
+	// configuration variables
+	FFrame3d Plane;
+
+	bool bEnableSnapToKnownLengths = true;
+
+	static constexpr int CardinalAxisTargetID = 10;
+	static constexpr int LastSegmentTargetID = 11;
+
+	int CardinalAxisPriority = 150;
+	int LastSegmentPriority = 140;
+	int KnownLengthPriorityDelta = 10;
+	int MinInternalPriority() const { return LastSegmentPriority - KnownLengthPriorityDelta; }
+
+public:
+	FPointPlanarSnapSolver();
+
+	void RegenerateTargetLines(bool bCardinalAxes, bool bLastHistorySegment);
+
+	virtual void Reset() override;
+	virtual void ResetActiveSnap() override;
+
+	void UpdatePointHistory(const TArray<FVector3d>& Points);
+	void UpdatePointHistory(const TArray<FVector>& Points);
+
+	void UpdateSnappedPoint(const FVector3d& PointIn);
+
+protected:
+
+	TArray<FSnapTargetLine> GeneratedLines;
+
+	TSet<int> IgnoreTargets;
+
+	TArray<FVector3d> PointHistory;
+
+	TArray<FSnapTargetPoint> GeneratedTargets;
+	void GenerateTargets(const FVector3d& PointIn);
+};
diff --git a/Engine/Plugins/Experimental/MeshModelingToolset/Source/ModelingComponents/Public/Snapping/RaySpatialSnapSolver.h b/Engine/Plugins/Experimental/MeshModelingToolset/Source/ModelingComponents/Public/Snapping/RaySpatialSnapSolver.h
new file mode 100644
--- /dev/null
+++ b/Engine/Plugins/Experimental/MeshModelingToolset/Source/ModelingComponents/Public/Snapping/RaySpatialSnapSolver.h
@@ -0,0 +1,39 @@
+// Copyright 1998-2019 Epic Games, Inc. All Rights Reserved.
+
+#pragma once
+
+#include "Snapping/BasePositionSnapSolver3.h"
+
+class FToolDataVisualizer;
+
+/**
+ * FRaySpatialSnapSolver solves for a Point snap location based on an input Ray
+ * and a set of snap targets (3D points and 3D lines). 
+ * 
+ * See FBasePositionSnapSolver3 for details on how to set up the snap problem
+ * and get results.
+ */
+class MODELINGCOMPONENTS_API FRaySpatialSnapSolver : public FBasePositionSnapSolver3
+{
+public:
+	FRaySpatialSnapSolver();
+
+	//
+	// solving
+	//
+
+	/** Solve the snapping problem */
+	void UpdateSnappedPoint(const FRay3d& Ray);
+
+	//
+	// Utility rendering
+	//
+
+	/** Visualization of snap targets and result (if available) */
+	void Draw(FToolDataVisualizer* Renderer, float LineLength, TMap<int,FLinearColor>* ColorMap = nullptr);
+
+protected:
+
+	TArray<FSnapTargetPoint> GeneratedTargetPoints;
+	void GenerateTargetPoints(const FRay3d& Ray);
+};
\ No newline at end of file
diff --git a/Engine/Plugins/Experimental/MeshModelingToolset/Source/ModelingComponents/Public/StaticMeshComponentBuilder.h b/Engine/Plugins/Experimental/MeshModelingToolset/Source/ModelingComponents/Public/StaticMeshComponentBuilder.h
new file mode 100644
--- /dev/null
+++ b/Engine/Plugins/Experimental/MeshModelingToolset/Source/ModelingComponents/Public/StaticMeshComponentBuilder.h
@@ -0,0 +1,47 @@
+// Copyright 1998-2019 Epic Games, Inc. All Rights Reserved.
+
+#pragma once
+
+#include "CoreMinimal.h"
+
+class AActor;
+class UStaticMesh;
+class UStaticMeshComponent;
+struct FMeshDescription;
+
+
+/**
+ * FStaticMeshComponentBuilder is a utility class that can build a 
+ * proper UStaticMeshComponent.
+ */
+class MODELINGCOMPONENTS_API FStaticMeshComponentBuilder
+{
+public:
+	/** We will create this UStaticMesh */
+	UStaticMesh* NewStaticMesh;
+#if WITH_EDITOR
+	/** This MeshDescription will be generated by calling NewStaticMesh->CreateMeshDescription(0) */
+	FMeshDescription *MeshDescription;
+#endif
+	/** We will create this component and set NewStaticMesh as it's mesh */
+	UStaticMeshComponent* NewMeshComponent;
+
+	FStaticMeshComponentBuilder();
+
+	/**
+	 * Call this to create and configure NewStaticMesh and MeshDescription
+	 * @param AssetPackage The asset package that the UStaticMesh will be saved into
+	 * @param MeshName the name of the mesh asset
+	 */
+	void Initialize(UPackage* AssetPackage, FName MeshName);
+
+	// (do your mesh setup here)
+
+	// 3) Create NewMeshComponent and add to Actor
+	/**
+	 * Create/Locate NewMeshComponent and assign NewStaticMesh to it, and then add it to the Actor.
+	 * If Actor is a AStaticMeshActor instance, then NewMeshComponent will be it's existing StaticMeshComponent, and we will assign NewStaticMesh to it.
+	 * Otherwise, we create new StaticMeshComponent, assign NewStaticMesh, and set it as the Root Component of the Actor
+	 */
+	void CreateAndSetAsRootComponent(AActor* Actor);
+};
diff --git a/Engine/Plugins/Experimental/MeshModelingToolset/Source/ModelingComponents/Public/Transforms/QuickAxisRotator.h b/Engine/Plugins/Experimental/MeshModelingToolset/Source/ModelingComponents/Public/Transforms/QuickAxisRotator.h
new file mode 100644
--- /dev/null
+++ b/Engine/Plugins/Experimental/MeshModelingToolset/Source/ModelingComponents/Public/Transforms/QuickAxisRotator.h
@@ -0,0 +1,107 @@
+// Copyright 1998-2019 Epic Games, Inc. All Rights Reserved.
+
+#pragma once
+
+#include "CoreMinimal.h"
+#include "QuickTransformer.h"
+#include "Snapping/RaySpatialSnapSolver.h"
+
+
+/**
+
+ */
+class MODELINGCOMPONENTS_API FQuickAxisRotator : public FQuickTransformer
+{
+public:
+	virtual ~FQuickAxisRotator() {}
+
+	/**
+	 * Set up internal data structures
+	 */
+	virtual void Initialize() override;
+
+	/**
+	 * Set current transform frame to the unit axes at the given Origin
+	 */
+	virtual void SetActiveFrameFromWorldAxes(const FVector3d& Origin) override;
+
+	/**
+	 * Set current transform frame to the given frame
+	 */
+	virtual void SetActiveWorldFrame(const FFrame3d& Frame) override;
+
+	/**
+	 * @return the current transform frame
+	 */
+	const FFrame3d& GetActiveWorldFrame() const 
+	{
+		return AxisFrameWorld; 
+	}
+
+	/**
+	 * Set current snap-axis frame to a frame at the given Origin with Z aligned to the given Normal.
+	 * @param bAlignToWorldAxes if true, the frame is rotated around the Normal so that the X and Y axes are best-aligned to the World axes
+	 */
+	virtual void SetActiveFrameFromWorldNormal(const FVector3d& Origin, const FVector3d& Normal, bool bAlignToWorldAxes);
+
+	/**
+	 * Update the current snap-axis frame with a new origin
+	 */
+	virtual void UpdateActiveFrameOrigin(const FVector3d& NewOrigin) override;
+
+
+	/**
+	 * Try to find the best snap point for the given Ray, and store in SnapPointOut if found
+	 * @return true if found
+	 */
+	virtual bool UpdateSnap(const FRay3d& Ray, FVector3d& SnapPointOut);
+
+	/** @return true if there is an active snap */
+	virtual bool HaveActiveSnap() const;
+
+
+	virtual bool HaveActiveSnapRotation() const;
+
+	virtual FFrame3d GetActiveRotationFrame() const;
+
+
+	void ClearAxisLock();
+	void SetAxisLock();
+	bool GetHaveLockedToAxis() const { return bHaveLockedToAxis; }
+
+
+	/**
+	 * Update internal copy of camera state. You must call this for snapping to work!
+	 */
+	virtual void UpdateCameraState(const FViewCameraState& CameraState) override;
+
+	/**
+	 * Draw a visualization of the current snap axes and active snap point
+	 * @param RenderAPI provide access to context rendering info
+	 */
+	virtual void Render(IToolsContextRenderAPI* RenderAPI) override;
+
+	/**
+	 * Draw a visualization of the current snap axes and active snap point
+	 * @param RenderAPI provide access to context rendering info
+	 */
+	virtual void PreviewRender(IToolsContextRenderAPI* RenderAPI) override;
+
+
+protected:
+	// camera state saved at last Render()
+	FViewCameraState CameraState;
+
+	bool bHaveLockedToAxis = false;
+	FIndex3i IgnoredAxes;
+
+	FFrame3d AxisFrameWorld;
+
+	FRaySpatialSnapSolver MoveAxisSolver;
+	void UpdateSnapAxes();
+
+	FToolDataVisualizer QuickAxisRenderer;
+	TMap<int, FLinearColor> AxisColorMap;
+
+	FToolDataVisualizer QuickAxisPreviewRenderer;
+};
\ No newline at end of file
diff --git a/Engine/Plugins/Experimental/MeshModelingToolset/Source/ModelingOperators/ModelingOperators.Build.cs b/Engine/Plugins/Experimental/MeshModelingToolset/Source/ModelingOperators/ModelingOperators.Build.cs
new file mode 100644
--- /dev/null
+++ b/Engine/Plugins/Experimental/MeshModelingToolset/Source/ModelingOperators/ModelingOperators.Build.cs
@@ -0,0 +1,76 @@
+// Copyright 1998-2019 Epic Games, Inc. All Rights Reserved.
+
+using UnrealBuildTool;
+
+public class ModelingOperators : ModuleRules
+{
+	public ModelingOperators(ReadOnlyTargetRules Target) : base(Target)
+	{
+		PCHUsage = ModuleRules.PCHUsageMode.UseExplicitOrSharedPCHs;
+		
+		PublicIncludePaths.AddRange(
+			new string[] {
+				// ... add public include paths required here ...
+			}
+			);
+				
+		
+		PrivateIncludePaths.AddRange(
+			new string[] {
+				// ... add other private include paths required here ...
+			}
+			);
+			
+		
+		PublicDependencyModuleNames.AddRange(
+			new string[]
+			{
+				// These dependencies were commented out when we split out ModelingOperatorsEditorOnly
+				// Many of them are editor only and should not be depended on here.  If you re-add any
+				// dependent modules, please confirm by "launching" or packaging that you are not introducing
+				// editor dependencies
+				//
+				//"MeshSolverUtilities",
+                //"ProxyLODMeshReduction", // currently required to be public due to IVoxelBasedCSG API
+                //"MeshDescription",
+
+				"Core",
+                "InteractiveToolsFramework",
+				"GeometricObjects",
+				"DynamicMesh",
+				"MeshConversion",
+                "GeometryAlgorithms", // required for constrained Delaunay triangulation
+                "MeshSolverUtilities", // required by the smoothing operators
+				// ... add other public dependencies that you statically link with here ...
+			}
+            );
+			
+		
+		PrivateDependencyModuleNames.AddRange(
+			new string[]
+			{
+				// These dependencies were commented out when we split out ModelingOperatorsEditorOnly
+				// Many of them are editor only and should not be depended on here.  If you re-add any
+				// dependent modules, please confirm by "launching" or packaging that you are not introducing
+				// editor dependencies
+				//
+				//"MeshBuilder",
+                //"MeshDescriptionOperations",
+				//"MeshUtilitiesCommon",
+				//"MeshReductionInterface", // for UE4 standard simplification
+
+				"CoreUObject",
+				"Engine",
+
+			}
+			);
+		
+		
+		DynamicallyLoadedModuleNames.AddRange(
+			new string[]
+			{
+				// ... add any modules that your module loads dynamically here ...
+			}
+			);
+	}
+}
diff --git a/Engine/Plugins/Experimental/MeshModelingToolset/Source/ModelingOperators/Private/MeshSpaceDeformerOps/BendMeshOp.cpp b/Engine/Plugins/Experimental/MeshModelingToolset/Source/ModelingOperators/Private/MeshSpaceDeformerOps/BendMeshOp.cpp
new file mode 100644
--- /dev/null
+++ b/Engine/Plugins/Experimental/MeshModelingToolset/Source/ModelingOperators/Private/MeshSpaceDeformerOps/BendMeshOp.cpp
@@ -0,0 +1,59 @@
+// Copyright 1998-2019 Epic Games, Inc. All Rights Reserved.
+#include "ModelingOperators\Public\SpaceDeformerOps\BendMeshOp.h"
+#include "ProfilingDebugging/ScopedTimers.h"
+
+
+//Bends along the Y-axis
+void FBendMeshOp::CalculateResult(FProgressCancel* Progress)
+{
+
+	//FString DebugLogString = FString::Printf(TEXT(" Bend timer "));
+	//FScopedDurationTimeLogger Timer(DebugLogString);
+	const double Curvature = GetModifierValue();
+	const double ToRadians = 0.017453292519943295;
+	const double ThetaRadians = ToRadians * Curvature;
+	const FVector3d O = ObjectSpaceToOpSpace * AxisOriginObjectSpace;
+	const double& Y0 = O.Y;
+	double YMin = Y0 + LowerBoundsInterval * AxesHalfLengths[2];
+	double YMax = Y0 + UpperBoundsInterval * AxesHalfLengths[2];
+	const double K = ThetaRadians / AxesHalfLengths[2];
+	const double Ik = 1.0 / K;
+	for (int VertexID : TargetMesh->VertexIndicesItr())
+	{
+		const FVector3d& OriginalPositionObjectSpace = OriginalPositions[VertexID];
+		if (Curvature == 0.0)
+		{
+			TargetMesh->SetVertex(VertexID, OriginalPositionObjectSpace);
+			continue;
+		}
+		const FVector3d Pos = ObjectSpaceToOpSpace * (OriginalPositionObjectSpace);
+
+		const double YHat = FMath::Clamp(Pos.Y, YMin, YMax);
+
+		const double Theta = K * (YHat - Y0);
+
+		const double S0 = FMath::Sin(Theta);
+		const double C0 = FMath::Cos(Theta);
+
+		const double ZP = Pos.Z - Ik;
+
+		const double X = Pos.X;
+		double Y = -S0 * (ZP) + Y0;
+		double Z =  C0 * (ZP) + Ik;
+
+		if (Pos.Y > YMax)
+		{
+			const double YDiff = Pos.Y - YMax;
+			Y += C0 * YDiff;
+			Z += S0 * YDiff;
+		}
+		else if (Pos.Y < YMin)
+		{
+			const double YDiff = Pos.Y - YMin;
+			Y += C0 * YDiff;
+			Z += S0 * YDiff;
+		}
+		TargetMesh->SetVertex(VertexID, OpSpaceToObjectSpace * FVector3d{ X,Y,Z });
+	}
+
+}
\ No newline at end of file
diff --git a/Engine/Plugins/Experimental/MeshModelingToolset/Source/ModelingOperators/Private/MeshSpaceDeformerOps/FlareMeshOp.cpp b/Engine/Plugins/Experimental/MeshModelingToolset/Source/ModelingOperators/Private/MeshSpaceDeformerOps/FlareMeshOp.cpp
new file mode 100644
--- /dev/null
+++ b/Engine/Plugins/Experimental/MeshModelingToolset/Source/ModelingOperators/Private/MeshSpaceDeformerOps/FlareMeshOp.cpp
@@ -0,0 +1,61 @@
+// Copyright 1998-2019 Epic Games, Inc. All Rights Reserved.
+#include "ModelingOperators\Public\SpaceDeformerOps\FlareMeshOp.h"
+
+
+
+//Some simple non-linear interpolation functions to play with.
+
+template <class T>
+T coserp(float percent, const T& value1, const T& value2)
+{
+	return T(0.5) * (cos(percent*PI) * (value1 - value2) + value1 + value2); 
+}
+
+template <class T>
+double inverseCoserp(const T& valueBetween, const T& value1, const T& value2)
+{
+	return acos((T(2) * valueBetween - value1 - value2) / (value1 - value2)) / PI;
+}
+
+template <class T>
+T sinerp(float percent, const T& value1, const T& value2)
+{
+	return T(0.5) * (sin(percent * PI) * (value1 - value2) + value1 + value2);   
+}
+
+template <class T>
+double  inverseSinerp(const T& valueBetween, const T& value1, const T& value2)
+{
+	return asin((T(2.0) * valueBetween - value1 - value2) / (value1 - value2)) / PI;
+}
+
+
+//Flares along the Z axis
+void FFlareMeshOp::CalculateResult(FProgressCancel* Progress)
+{
+	FMatrix3d Rotate{ FVector3d{0,1,0}, FVector3d{0,0,1}, FVector3d{1,0,0}, false };
+	FMatrix3d ToOpSpace = Rotate * ObjectSpaceToOpSpace;
+
+	const FVector3d O = ToOpSpace * AxisOriginObjectSpace;
+	const double& Z0 = O.Z;
+	const double ZMin = Z0 + LowerBoundsInterval * AxesHalfLengths[2];
+	const double ZMax = Z0 + UpperBoundsInterval * AxesHalfLengths[2];
+
+	FMatrix3d FromOpSpace = ToOpSpace.Inverse();
+	for (int VertexID : TargetMesh->VertexIndicesItr())
+	{
+		const FVector3d& OriginalPositionObjectSpace = OriginalPositions[VertexID];
+
+		const FVector3d Pos = ToOpSpace * (OriginalPositionObjectSpace);
+
+		double T = FMath::Clamp(FMath::Abs((Z0 - Pos.Z)) / (ZMax - ZMin) * 0.5, 0.0, 1.0);
+		double R = coserp((1.0-T), 1.0, GetModifierValue());
+
+		double X = Pos.X * R;
+		double Y = Pos.Y * R;
+		double Z = Pos.Z;
+
+		TargetMesh->SetVertex(VertexID, FromOpSpace * (FVector3d{ X,Y,Z }));
+	}
+
+}
\ No newline at end of file
diff --git a/Engine/Plugins/Experimental/MeshModelingToolset/Source/ModelingOperators/Private/MeshSpaceDeformerOps/MeshSpaceDeformerOp.cpp b/Engine/Plugins/Experimental/MeshModelingToolset/Source/ModelingOperators/Private/MeshSpaceDeformerOps/MeshSpaceDeformerOp.cpp
new file mode 100644
--- /dev/null
+++ b/Engine/Plugins/Experimental/MeshModelingToolset/Source/ModelingOperators/Private/MeshSpaceDeformerOps/MeshSpaceDeformerOp.cpp
@@ -0,0 +1,32 @@
+// Copyright 1998-2019 Epic Games, Inc. All Rights Reserved.
+
+#pragma once
+
+#include "SpaceDeformerOps\MeshSpaceDeformerOp.h"
+#include "DynamicMesh3.h"
+
+void FMeshSpaceDeformerOp::CalculateResult(FProgressCancel* Progress)
+{
+}
+
+void FMeshSpaceDeformerOp::UpdateMesh(FDynamicMesh3* TargetMeshComponent)
+{
+	TargetMesh = TargetMeshComponent;
+
+	OriginalPositions.SetNumUninitialized(TargetMesh->MaxVertexID());
+	for (int32 VertexID : TargetMesh->VertexIndicesItr())
+	{
+		OriginalPositions[VertexID] = TargetMesh->GetVertex(VertexID);
+	}
+}
+
+void FMeshSpaceDeformerOp::UpdateAxisData(const FMatrix3d & ObjectSpaceToOpSpaceTransform, const FVector3d & ObjectSpaceOrigin, const FVector3d & AxesHalfExtents, double LowerBounds, double UpperBounds, double Modifier)
+{
+	ObjectSpaceToOpSpace = ObjectSpaceToOpSpaceTransform;  // Rotation matrix to rotate from object space to the space expected by the operator
+	OpSpaceToObjectSpace = ObjectSpaceToOpSpace.Inverse(); // Pre-computation for use in the operator
+	AxisOriginObjectSpace = ObjectSpaceOrigin;			   // The location of the handle's origin, transformed into the object space of the target mesh
+	AxesHalfLengths = AxesHalfExtents;					   // Each element of this vector is a scalar corresponding to the distance from the centroid to the farthest vertex in that respective direction
+ 	ModifierPercent = Modifier;							   // Input from the tool, as each operator will have a min/max value i.e. the angle of curvature, twist or scale. This is the percent used to interpolate between the min/max value
+	LowerBoundsInterval = LowerBounds;					   // The interval [0.0-1.0] representing the range of effect in space this operator will have
+	UpperBoundsInterval = UpperBounds;				
+}
diff --git a/Engine/Plugins/Experimental/MeshModelingToolset/Source/ModelingOperators/Private/MeshSpaceDeformerOps/TwistMeshOp.cpp b/Engine/Plugins/Experimental/MeshModelingToolset/Source/ModelingOperators/Private/MeshSpaceDeformerOps/TwistMeshOp.cpp
new file mode 100644
--- /dev/null
+++ b/Engine/Plugins/Experimental/MeshModelingToolset/Source/ModelingOperators/Private/MeshSpaceDeformerOps/TwistMeshOp.cpp
@@ -0,0 +1,38 @@
+// Copyright 1998-2019 Epic Games, Inc. All Rights Reserved.
+
+#include "ModelingOperators\Public\SpaceDeformerOps\TwistMeshOp.h"
+
+//Twists along the Z axis
+void FTwistMeshOp::CalculateResult(FProgressCancel* Progress)
+{
+	FMatrix3d Rotate{FVector3d{0,1,0}, FVector3d{0,0,1}, FVector3d{1,0,0}, false};
+	FMatrix3d ToOpSpace = Rotate * ObjectSpaceToOpSpace;
+
+	const FVector3d O = ToOpSpace * AxisOriginObjectSpace;
+	const double& Z0 = O.Z;
+	const double ZMin = Z0 + LowerBoundsInterval * AxesHalfLengths[2];
+	const double ZMax = Z0 + UpperBoundsInterval * AxesHalfLengths[2];
+
+	const double ToRadians = 0.017453292519943295;
+	const double ThetaRadians = ToRadians * GetModifierValue();
+
+	FMatrix3d FromOpSpace = ToOpSpace.Inverse();
+	for (int VertexID : TargetMesh->VertexIndicesItr())
+	{
+		const FVector3d& OriginalPositionObjectSpace = OriginalPositions[VertexID];
+
+		const FVector3d Pos = ToOpSpace * (OriginalPositionObjectSpace);
+
+		double T = FMath::Clamp(( ZMax - Pos.Z ) / ( ZMax - ZMin ), 0.0, 1.0);
+		double Theta = ThetaRadians * T;
+
+		const double C0 = FMath::Cos(Theta);
+		const double S0 = FMath::Sin(Theta);
+
+		double X = Pos.X * C0 + Pos.Y * S0;
+		double Y = Pos.X * (-S0) + Pos.Y * C0;
+		double Z = Pos.Z;
+
+		TargetMesh->SetVertex(VertexID, FromOpSpace * FVector3d{ X,Y,Z });
+	}
+}
\ No newline at end of file
diff --git a/Engine/Plugins/Experimental/MeshModelingToolset/Source/ModelingOperators/Private/ParameterizationOps/UVProjectionOp.cpp b/Engine/Plugins/Experimental/MeshModelingToolset/Source/ModelingOperators/Private/ParameterizationOps/UVProjectionOp.cpp
new file mode 100644
--- /dev/null
+++ b/Engine/Plugins/Experimental/MeshModelingToolset/Source/ModelingOperators/Private/ParameterizationOps/UVProjectionOp.cpp
@@ -0,0 +1,243 @@
+// Copyright 1998-2019 Epic Games, Inc. All Rights Reserved.
+
+#pragma once
+
+#include "ParameterizationOps/UVProjectionOp.h"
+
+#include "Engine/StaticMesh.h"
+
+#include "MeshDescriptionToDynamicMesh.h"
+#include "MeshSimplification.h"
+#include "MeshConstraintsUtil.h"
+#include "MatrixTypes.h"
+
+#include "Async/ParallelFor.h"
+
+#include "MeshNormals.h"
+#include "DynamicMeshEditor.h"
+
+
+void FUVProjectionOp::SetTransform(const FTransform& Transform) {
+	ResultTransform = (FTransform3d)Transform;
+}
+
+void FUVProjectionOp::CalculateResult(FProgressCancel* Progress)
+{
+	if (Progress->Cancelled())
+	{
+		return;
+	}
+	bool bDiscardAttributes = false;
+	ResultMesh->Copy(*OriginalMesh, true, true, true, !bDiscardAttributes);
+
+	if (!ensureMsgf(ResultMesh->HasAttributes(), TEXT("Attributes not found on mesh? Conversion should always create them, so this operator should not need to do so.")))
+	{
+		ResultMesh->EnableAttributes();
+	}
+
+	// TODO: get transform from mesh vertex space to projection primitive space
+	FTransform ToProjPrimXF = FTransform(ResultTransform) * ProjectionTransform.Inverse();
+	FTransform3d ToProjPrim(ToProjPrimXF);
+
+	TArray<FVector3d> TransformedVertices; TransformedVertices.SetNumUninitialized(ResultMesh->VertexCount());
+	ParallelFor(ResultMesh->VertexCount(), [this, &ToProjPrim, &TransformedVertices](int32 VertexID)
+	{
+		if (ResultMesh->IsVertex(VertexID))
+		{
+			TransformedVertices[VertexID] = ToProjPrim.TransformPosition(ResultMesh->GetVertex(VertexID));
+		}
+	});
+	
+	
+	int NumUVLayers = ResultMesh->Attributes()->NumUVLayers();
+	int LayerIndex = 0; // TODO: support non-zero UV layer, make this a member of FUVProjectionOp and a setting in the UI (will require updating dmesh attribute code)
+	FDynamicMeshUVOverlay* UVLayer = ResultMesh->Attributes()->GetUVLayer(LayerIndex);
+	UVLayer->ClearElements();
+
+	FVector2D Scale = FVector2D(.5,.5)*UVScale, Offset = FVector2D(.5,.5)+UVOffset;
+	// project to major axis
+	auto ProjAxis = [&Scale, &Offset](const FVector3d& P, int Ax1, int Ax2, float Ax1Scale, float Ax2Scale)
+	{
+		return FVector2f(float(P[Ax1]) * Ax1Scale * Scale.X + Offset.X, float(P[Ax2]) * Ax2Scale * Scale.Y + Offset.Y);
+	};
+
+	int MaxVertexID = ResultMesh->MaxVertexID();
+	int MaxTriangleID = ResultMesh->MaxTriangleID();
+	if (ProjectionMethod == EUVProjectionMethod::Plane)
+	{
+		UVLayer->BeginUnsafeElementsInsert();
+		// 1:1 with vertices; note we cannot just delete the UV layer and switch to vertex UVs because that would delete the normals layer too
+		for (int VID = MaxVertexID-1; VID >= 0; VID--)
+		{
+			if (ResultMesh->IsVertex(VID))
+			{
+				FVector2f UV = ProjAxis(TransformedVertices[VID], 0, 1, 1, 1);
+				UVLayer->InsertElement(VID, UV, VID, true);
+			}
+		}
+		UVLayer->EndUnsafeElementsInsert();
+		for (int TID = 0; TID < MaxTriangleID; TID++)
+		{
+			if (ResultMesh->IsTriangle(TID))
+			{
+				UVLayer->SetTriangle(TID, ResultMesh->GetTriangle(TID));
+			}
+		}
+	}
+	else
+	{	
+		// All cases in this branch require normals 
+		// compute normals on the transformed vertices
+		TArray<FVector3d> TransformedNormals; TransformedNormals.SetNumUninitialized(ResultMesh->TriangleCount());
+		ParallelFor(ResultMesh->TriangleCount(), [this, &TransformedVertices, &TransformedNormals](int32 TriangleID)
+		{
+			if (ResultMesh->IsTriangle(TriangleID))
+			{
+				FIndex3i Triangle = ResultMesh->GetTriangle(TriangleID);
+				TransformedNormals[TriangleID] = VectorUtil::Normal(TransformedVertices[Triangle.A], TransformedVertices[Triangle.B], TransformedVertices[Triangle.C]);
+			}
+		});
+		TMap<FIndex2i, int> VertIdxAndBucketIDToElementID;
+
+		if (ProjectionMethod == EUVProjectionMethod::Cube)
+		{
+			int Minor1s[3] = { 1, 0, 0 };
+			int Minor2s[3] = { 2, 2, 1 };
+			int Minor1Flip[3] = { -1, 1, 1 };
+			int Minor2Flip[3] = { -1, -1, 1 };
+			for (int TID = 0; TID < MaxTriangleID; TID++)
+			{
+				if (!ResultMesh->IsTriangle(TID))
+				{
+					continue;
+				}
+				const FVector3d& N = TransformedNormals[TID];
+				FVector3d NAbs(FMathd::Abs(N.X), FMathd::Abs(N.Y), FMathd::Abs(N.Z));
+				int MajorAxis = NAbs[0] > NAbs[1] ? (NAbs[0] > NAbs[2] ? 0 : 2) : (NAbs[1] > NAbs[2] ? 1 : 2);
+				double MajorAxisSign = FMathd::Sign(N[MajorAxis]);
+				int Minor1 = Minor1s[MajorAxis];
+				int Minor2 = Minor2s[MajorAxis];
+				int Bucket = MajorAxis;
+				if (MajorAxisSign > 0)
+				{
+					Bucket += 3;
+				}
+				FIndex3i OverlayTri;
+				FIndex3i SourceTri = ResultMesh->GetTriangle(TID);
+				for (int SubIdx = 0; SubIdx < 3; SubIdx++)
+				{
+					FIndex2i ElementKey(SourceTri[SubIdx], Bucket);
+					const int* ElementIdx = VertIdxAndBucketIDToElementID.Find(ElementKey);
+					if (!ElementIdx)
+					{
+						FVector3d LocalV = TransformedVertices[SourceTri[SubIdx]];
+						int NewElementIdx = UVLayer->AppendElement(ProjAxis(LocalV, Minor1, Minor2, MajorAxisSign*Minor1Flip[MajorAxis], Minor2Flip[MajorAxis]), ElementKey.A);
+						VertIdxAndBucketIDToElementID.Add(ElementKey, NewElementIdx);
+						OverlayTri[SubIdx] = NewElementIdx;
+					}
+					else
+					{
+						OverlayTri[SubIdx] = *ElementIdx;
+					}
+				}
+				UVLayer->SetTriangle(TID, OverlayTri);
+			}
+		}
+		else if (ProjectionMethod == EUVProjectionMethod::Cylinder)
+		{
+			double DotThresholdRejectFromPlane = FMathd::Cos(FMathd::DegToRad * CylinderProjectToTopOrBottomAngleThreshold);
+			for (int TID = 0; TID < MaxTriangleID; TID++)
+			{
+				if (!ResultMesh->IsTriangle(TID))
+				{
+					continue;
+				}
+				const FVector3d& N = TransformedNormals[TID];
+				FIndex3i OverlayTri;
+				FIndex3i SourceTri = ResultMesh->GetTriangle(TID);
+				if (FMathd::Abs(N.Z) > DotThresholdRejectFromPlane)
+				{
+					int MajorAxis = 2;
+					double MajorAxisSign = FMathd::Sign(N[MajorAxis]);
+					// project to +/- Z
+					int Bucket = MajorAxisSign > 0 ? 1 : 0;
+					int Minor1 = 0;
+					int Minor2 = 1;
+
+					// TODO: extract common major axis plane projection sub-function from cube projection & this?
+					for (int SubIdx = 0; SubIdx < 3; SubIdx++)
+					{
+						FIndex2i ElementKey(SourceTri[SubIdx], Bucket);
+						const int* ElementIdx = VertIdxAndBucketIDToElementID.Find(ElementKey);
+						if (!ElementIdx)
+						{
+							FVector3d LocalV = TransformedVertices[SourceTri[SubIdx]];
+							int NewElementIdx = UVLayer->AppendElement(ProjAxis(LocalV, Minor1, Minor2, MajorAxisSign, 1), ElementKey.A);
+							VertIdxAndBucketIDToElementID.Add(ElementKey, NewElementIdx);
+							OverlayTri[SubIdx] = NewElementIdx;
+						}
+						else
+						{
+							OverlayTri[SubIdx] = *ElementIdx;
+						}
+					}
+				}
+				else
+				{
+					// project to cylinder
+					int Bucket = 3;
+					FVector3d Centroid = (TransformedVertices[SourceTri.A] + TransformedVertices[SourceTri.B] + TransformedVertices[SourceTri.C]) * (1.0 / 3.0);
+					double CentroidAngle = FMathd::Atan2Positive(Centroid.Y, Centroid.X);
+					for (int SubIdx = 0; SubIdx < 3; SubIdx++)
+					{
+						FVector3d LocalV = TransformedVertices[SourceTri[SubIdx]];
+						double VAngle = FMathd::Atan2Positive(LocalV.Y, LocalV.X);
+						double AngleDiff = VAngle - CentroidAngle;
+						if (FMathd::Abs(AngleDiff) > FMathd::Pi) {
+							if (AngleDiff < 0)
+							{
+								VAngle += FMathd::TwoPi;
+								Bucket--;
+							}
+							else
+							{
+								VAngle -= FMathd::TwoPi;
+								Bucket++;
+							}
+						}
+						FIndex2i ElementKey(SourceTri[SubIdx], Bucket);
+						const int* ElementIdx = VertIdxAndBucketIDToElementID.Find(ElementKey);
+						if (!ElementIdx)
+						{
+							FVector2f UV(-(float(VAngle)*FMathf::InvPi - 1.0f)*Scale.X + Offset.X, -float(LocalV.Z)*Scale.Y + Offset.Y);
+							int NewElementIdx = UVLayer->AppendElement(UV, ElementKey.A);
+							VertIdxAndBucketIDToElementID.Add(ElementKey, NewElementIdx);
+							OverlayTri[SubIdx] = NewElementIdx;
+						}
+						else
+						{
+							OverlayTri[SubIdx] = *ElementIdx;
+						}
+					}
+				}
+
+				UVLayer->SetTriangle(TID, OverlayTri);
+			}
+		}
+
+
+		// rescale UVs to an average world-space target length (and then re-multiply by the UVScale param)
+		if (bWorldSpaceUVScale)
+		{
+			FDynamicMeshEditor Editor(ResultMesh.Get());
+			// global rescale with a 'standard' global scale factor
+			Editor.RescaleAttributeUVs(.01, true, 0, ResultTransform);
+			// apply the 2D uvscale after
+			for (int ElID : UVLayer->ElementIndicesItr())
+			{
+				UVLayer->SetElement(ElID, UVLayer->GetElement(ElID) * UVScale);
+			}
+		}
+	}
+
+}
diff --git a/Engine/Plugins/Experimental/MeshModelingToolset/Source/ModelingOperators/Private/SmoothingOps/CotanSmoothingOp.cpp b/Engine/Plugins/Experimental/MeshModelingToolset/Source/ModelingOperators/Private/SmoothingOps/CotanSmoothingOp.cpp
new file mode 100644
--- /dev/null
+++ b/Engine/Plugins/Experimental/MeshModelingToolset/Source/ModelingOperators/Private/SmoothingOps/CotanSmoothingOp.cpp
@@ -0,0 +1,28 @@
+// Copyright 1998-2019 Epic Games, Inc. All Rights Reserved.
+
+#include "SmoothingOps/CotanSmoothingOp.h"
+
+#include "MeshSmoothingUtilities.h"
+
+
+FCotanSmoothingOp::FCotanSmoothingOp(const FDynamicMesh3* Mesh, float Speed, int32 Iterations) :
+	FSmoothingOpBase(Mesh, Speed, Iterations)
+{
+}
+
+
+void FCotanSmoothingOp::CalculateResult(FProgressCancel* Progress)
+{
+	// Update the values in the position buffer with smoothed positions.
+	Smooth();
+
+	// Copy the results back into the result mesh and update normals
+	UpdateResultMesh();
+
+}
+
+void FCotanSmoothingOp::Smooth()
+{
+	double Intensity = 1.;
+	MeshSmoothingOperators::ComputeSmoothing_BiHarmonic(ELaplacianWeightScheme::ClampedCotangent, *ResultMesh, SmoothSpeed, Intensity, SmoothIterations, PositionBuffer);
+}
diff --git a/Engine/Plugins/Experimental/MeshModelingToolset/Source/ModelingOperators/Private/SmoothingOps/IterativeSmoothingOp.cpp b/Engine/Plugins/Experimental/MeshModelingToolset/Source/ModelingOperators/Private/SmoothingOps/IterativeSmoothingOp.cpp
new file mode 100644
--- /dev/null
+++ b/Engine/Plugins/Experimental/MeshModelingToolset/Source/ModelingOperators/Private/SmoothingOps/IterativeSmoothingOp.cpp
@@ -0,0 +1,68 @@
+// Copyright 1998-2019 Epic Games, Inc. All Rights Reserved.
+
+#include "SmoothingOps/IterativeSmoothingOp.h"
+#include "Async/ParallelFor.h"
+
+FIterativeSmoothingOp::FIterativeSmoothingOp(const FDynamicMesh3* Mesh, float Speed, int32 Iterations) :
+	FSmoothingOpBase(Mesh, Speed, Iterations)
+{
+	SmoothedBuffer = PositionBuffer;
+}
+
+void FIterativeSmoothingOp::CalculateResult(FProgressCancel* Progress)
+{
+	// Update the values in the position buffer with smoothed positions.
+	Smooth();
+
+	// Copy the results back into the result mesh and update normals
+	UpdateResultMesh();
+
+}
+
+void FIterativeSmoothingOp::Smooth()
+{
+
+	// Each iteration is equivalent to time advancing the diffusion system 
+	// 
+	//  d p / dt = k*k L[ p ]
+	// 
+	// with forward Euler
+	//  
+	// p^{n+1} = p^{n} + dt k * k L[ p^{n} ]
+	// 
+	// Where { L[p] }_i = Sum[ w_{ij} p_j ]  
+	//        
+	//        and w_{ii} = - Sum[ w_{ij}, j != i }
+	//
+	// Here SmoothSpeed = -dt * k * k * w_{ii}
+	// 
+	// NB: Here -w_{ii} is the valence of the ith vertex.  
+	//     and 'k' may be spatially dependent. 
+
+	int NV = ResultMesh->MaxVertexID();
+
+	for (int k = 0; k < SmoothIterations; ++k)
+	{
+		ParallelFor(NV, [=](int vid) {
+			if (ResultMesh->IsVertex(vid) == false)
+			{
+				return;
+			}
+
+			FVector3d Centroid(0, 0, 0);
+			int NumNbrs = 0;
+			for (int nbrvid : ResultMesh->VtxVerticesItr(vid))
+			{
+				Centroid += PositionBuffer[nbrvid];
+				NumNbrs++;
+			}
+			Centroid /= (double)NumNbrs;
+			SmoothedBuffer[vid] = FVector3d::Lerp(PositionBuffer[vid], Centroid, SmoothSpeed);
+		}, false);
+
+		for (int vid = 0; vid < NV; ++vid)
+		{
+			PositionBuffer[vid] = SmoothedBuffer[vid];
+		}
+	}
+}
diff --git a/Engine/Plugins/Experimental/MeshModelingToolset/Source/ModelingOperators/Private/SmoothingOps/MeanValueSmoothingOp.cpp b/Engine/Plugins/Experimental/MeshModelingToolset/Source/ModelingOperators/Private/SmoothingOps/MeanValueSmoothingOp.cpp
new file mode 100644
--- /dev/null
+++ b/Engine/Plugins/Experimental/MeshModelingToolset/Source/ModelingOperators/Private/SmoothingOps/MeanValueSmoothingOp.cpp
@@ -0,0 +1,28 @@
+// Copyright 1998-2019 Epic Games, Inc. All Rights Reserved.
+
+#include "SmoothingOps/MeanValueSmoothingOp.h"
+
+#include "MeshSmoothingUtilities.h"
+
+
+FMeanValueSmoothingOp::FMeanValueSmoothingOp(const FDynamicMesh3* Mesh, float Speed, int32 Iterations) :
+	FSmoothingOpBase(Mesh, Speed, Iterations)
+{
+}
+
+
+void FMeanValueSmoothingOp::CalculateResult(FProgressCancel* Progress) 
+{
+	// Update the values in the position buffer with smoothed positions.
+	Smooth();
+
+	// Copy the results back into the result mesh and update normals
+	UpdateResultMesh();
+
+}
+
+void FMeanValueSmoothingOp::Smooth()
+{
+	double Intensity = 1.;
+	MeshSmoothingOperators::ComputeSmoothing_BiHarmonic(ELaplacianWeightScheme::MeanValue, *ResultMesh, SmoothSpeed, Intensity, SmoothIterations, PositionBuffer);
+}
\ No newline at end of file
diff --git a/Engine/Plugins/Experimental/MeshModelingToolset/Source/ModelingOperators/Private/SmoothingOps/SmoothingOpBase.cpp b/Engine/Plugins/Experimental/MeshModelingToolset/Source/ModelingOperators/Private/SmoothingOps/SmoothingOpBase.cpp
new file mode 100644
--- /dev/null
+++ b/Engine/Plugins/Experimental/MeshModelingToolset/Source/ModelingOperators/Private/SmoothingOps/SmoothingOpBase.cpp
@@ -0,0 +1,52 @@
+// Copyright 1998-2019 Epic Games, Inc. All Rights Reserved.
+
+#include "SmoothingOps/SmoothingOpBase.h"
+
+#include "DynamicMesh3.h"
+#include "MeshNormals.h"
+
+
+FSmoothingOpBase::FSmoothingOpBase(const FDynamicMesh3* Mesh, float Speed, int32 Iterations) :
+		FDynamicMeshOperator(),
+		SmoothSpeed(Speed),
+		SmoothIterations(Iterations)
+{
+	// deep copy the src mesh into the result mesh.  This ResultMesh will be directly updated by the smoothing.
+	ResultMesh->Copy(*Mesh);
+	
+	PositionBuffer.SetNum(ResultMesh->MaxVertexID());
+
+	for (int vid : ResultMesh->VertexIndicesItr())
+	{
+		PositionBuffer[vid] = ResultMesh->GetVertex(vid);
+	}
+}
+
+void FSmoothingOpBase::SetTransform(FTransform3d& XForm)
+{
+	ResultTransform = XForm;
+}
+
+
+void FSmoothingOpBase::UpdateResultMesh()
+{
+	// move the verts to the new location
+	for (int32 vid : ResultMesh->VertexIndicesItr())
+	{
+		const FVector3d Pos = PositionBuffer[vid];
+		ResultMesh->SetVertex(vid, Pos);
+	}
+
+	// recalculate normals
+	if (ResultMesh->HasAttributes())
+	{
+		FMeshNormals Normals(ResultMesh.Get());
+		FDynamicMeshNormalOverlay* NormalOverlay = ResultMesh->Attributes()->PrimaryNormals();
+		Normals.RecomputeOverlayNormals(NormalOverlay);
+		Normals.CopyToOverlay(NormalOverlay);
+	}
+	else
+	{
+		FMeshNormals::QuickComputeVertexNormals(*ResultMesh);
+	}
+}
diff --git a/Engine/Plugins/Experimental/MeshModelingToolset/Source/ModelingOperators/Public/CleaningOps/EditNormalsOp.h b/Engine/Plugins/Experimental/MeshModelingToolset/Source/ModelingOperators/Public/CleaningOps/EditNormalsOp.h
new file mode 100644
--- /dev/null
+++ b/Engine/Plugins/Experimental/MeshModelingToolset/Source/ModelingOperators/Public/CleaningOps/EditNormalsOp.h
@@ -0,0 +1,47 @@
+// Copyright 1998-2019 Epic Games, Inc. All Rights Reserved.
+
+#pragma once
+
+#include "CoreMinimal.h"
+#include "Util/ProgressCancel.h"
+#include "ModelingOperators.h"
+
+#include "EditNormalsOp.generated.h"
+
+
+UENUM()
+enum class ENormalCalculationMethod : uint8
+{
+	AreaWeighted,
+	AngleWeighted,
+	AreaAngleWeighting
+};
+
+
+
+class MODELINGOPERATORS_API FEditNormalsOp : public FDynamicMeshOperator
+{
+public:
+	virtual ~FEditNormalsOp() {}
+
+	// inputs
+	TSharedPtr<FDynamicMesh3> OriginalMesh;
+
+	bool bFixInconsistentNormals;
+	bool bInvertNormals;
+	bool bRecomputeNormals;
+	ENormalCalculationMethod NormalCalculationMethod;
+	bool bSplitNormals;
+	bool bAllowSharpVertices;
+	float NormalSplitThreshold;
+
+	void SetTransform(const FTransform& Transform);
+
+	//
+	// FDynamicMeshOperator implementation
+	// 
+
+	virtual void CalculateResult(FProgressCancel* Progress) override;
+};
+
+
diff --git a/Engine/Plugins/Experimental/MeshModelingToolset/Source/ModelingOperators/Public/ParameterizationOps/UVProjectionOp.h b/Engine/Plugins/Experimental/MeshModelingToolset/Source/ModelingOperators/Public/ParameterizationOps/UVProjectionOp.h
new file mode 100644
--- /dev/null
+++ b/Engine/Plugins/Experimental/MeshModelingToolset/Source/ModelingOperators/Public/ParameterizationOps/UVProjectionOp.h
@@ -0,0 +1,45 @@
+// Copyright 1998-2019 Epic Games, Inc. All Rights Reserved.
+
+#pragma once
+
+#include "CoreMinimal.h"
+#include "Util/ProgressCancel.h"
+#include "ModelingOperators.h"
+
+#include "UVProjectionOp.generated.h"
+
+
+UENUM()
+enum class EUVProjectionMethod : uint8
+{
+	Cube,
+	Cylinder,
+	Plane
+};
+
+
+
+class MODELINGOPERATORS_API FUVProjectionOp : public FDynamicMeshOperator
+{
+public:
+	virtual ~FUVProjectionOp() {}
+
+	// inputs
+	TSharedPtr<FDynamicMesh3> OriginalMesh;
+
+	EUVProjectionMethod ProjectionMethod;
+	float CylinderProjectToTopOrBottomAngleThreshold;
+	FTransform ProjectionTransform;
+	FVector2D UVScale, UVOffset;
+	bool bWorldSpaceUVScale;
+
+	void SetTransform(const FTransform& Transform);
+
+	//
+	// FDynamicMeshOperator implementation
+	// 
+
+	virtual void CalculateResult(FProgressCancel* Progress) override;
+};
+
+
diff --git a/Engine/Plugins/Experimental/MeshModelingToolset/Source/ModelingOperators/Public/SmoothingOps/CotanSmoothingOp.h b/Engine/Plugins/Experimental/MeshModelingToolset/Source/ModelingOperators/Public/SmoothingOps/CotanSmoothingOp.h
new file mode 100644
--- /dev/null
+++ b/Engine/Plugins/Experimental/MeshModelingToolset/Source/ModelingOperators/Public/SmoothingOps/CotanSmoothingOp.h
@@ -0,0 +1,21 @@
+// Copyright 1998-2019 Epic Games, Inc. All Rights Reserved.
+
+#pragma once
+
+#include "SmoothingOpBase.h"
+#include "CoreMinimal.h"
+
+
+class MODELINGOPERATORS_API  FCotanSmoothingOp : public FSmoothingOpBase
+{
+public:
+	FCotanSmoothingOp(const FDynamicMesh3* Mesh, float Speed, int32 Iterations);
+
+	~FCotanSmoothingOp() override {};
+
+	void CalculateResult(FProgressCancel* Progress) override;
+
+private:
+	// Compute the smoothed result by using Cotan Biharmonic
+	void Smooth();	
+};
\ No newline at end of file
diff --git a/Engine/Plugins/Experimental/MeshModelingToolset/Source/ModelingOperators/Public/SmoothingOps/IterativeSmoothingOp.h b/Engine/Plugins/Experimental/MeshModelingToolset/Source/ModelingOperators/Public/SmoothingOps/IterativeSmoothingOp.h
new file mode 100644
--- /dev/null
+++ b/Engine/Plugins/Experimental/MeshModelingToolset/Source/ModelingOperators/Public/SmoothingOps/IterativeSmoothingOp.h
@@ -0,0 +1,26 @@
+// Copyright 1998-2019 Epic Games, Inc. All Rights Reserved.
+
+#pragma once
+
+#include "SmoothingOpBase.h"
+#include "CoreMinimal.h"
+#include "VectorTypes.h"
+
+class MODELINGOPERATORS_API  FIterativeSmoothingOp : public FSmoothingOpBase
+{
+public:
+	FIterativeSmoothingOp(const FDynamicMesh3* Mesh, float Speed, int32 Iterations);
+
+	~FIterativeSmoothingOp() override {};
+
+	// Apply smoothing. results in an updated ResultMesh
+	void CalculateResult(FProgressCancel* Progress) override;
+
+private:
+
+	// Does the actual interative smoothign. 
+	void Smooth();
+	
+private:
+	TArray<FVector3d> SmoothedBuffer;
+};
\ No newline at end of file
diff --git a/Engine/Plugins/Experimental/MeshModelingToolset/Source/ModelingOperators/Public/SmoothingOps/MeanValueSmoothingOp.h b/Engine/Plugins/Experimental/MeshModelingToolset/Source/ModelingOperators/Public/SmoothingOps/MeanValueSmoothingOp.h
new file mode 100644
--- /dev/null
+++ b/Engine/Plugins/Experimental/MeshModelingToolset/Source/ModelingOperators/Public/SmoothingOps/MeanValueSmoothingOp.h
@@ -0,0 +1,21 @@
+// Copyright 1998-2019 Epic Games, Inc. All Rights Reserved.
+
+#pragma once
+
+#include "SmoothingOpBase.h"
+#include "CoreMinimal.h"
+
+class MODELINGOPERATORS_API  FMeanValueSmoothingOp : public FSmoothingOpBase
+{
+public:
+	FMeanValueSmoothingOp(const FDynamicMesh3* Mesh, float Speed, int32 Iterations); 
+	
+	~FMeanValueSmoothingOp() override {};
+
+	void CalculateResult(FProgressCancel* Progress) override;
+
+private:
+	// computed the smoothed result using mean value biharmonic operator
+	void Smooth();
+	
+};
\ No newline at end of file
diff --git a/Engine/Plugins/Experimental/MeshModelingToolset/Source/ModelingOperators/Public/SmoothingOps/SmoothingOpBase.h b/Engine/Plugins/Experimental/MeshModelingToolset/Source/ModelingOperators/Public/SmoothingOps/SmoothingOpBase.h
new file mode 100644
--- /dev/null
+++ b/Engine/Plugins/Experimental/MeshModelingToolset/Source/ModelingOperators/Public/SmoothingOps/SmoothingOpBase.h
@@ -0,0 +1,35 @@
+// Copyright 1998-2019 Epic Games, Inc. All Rights Reserved.
+
+#pragma once
+
+#include "CoreMinimal.h"
+#include "Util/ProgressCancel.h"
+#include "ModelingOperators.h"
+#include "VectorTypes.h"
+
+
+class FDynamicMesh3;
+
+class MODELINGOPERATORS_API FSmoothingOpBase : public FDynamicMeshOperator
+{
+public:
+	FSmoothingOpBase(const FDynamicMesh3* Mesh, float Speed, int32 Iterations);
+
+	virtual ~FSmoothingOpBase() override {}
+
+	// set ability on protected transform.
+	void SetTransform(FTransform3d& XForm);
+
+	// base class overrides this.  Results in updated ResultMesh.
+	virtual void CalculateResult(FProgressCancel* Progress) override = 0;
+
+	// copy the PositionBuffer locations back to the ResultMesh and recompute normal if it exists.
+	void UpdateResultMesh();
+
+protected:
+	float SmoothSpeed;
+	int32 SmoothIterations;
+
+	TArray<FVector3d> PositionBuffer;
+
+};
\ No newline at end of file
diff --git a/Engine/Plugins/Experimental/MeshModelingToolset/Source/ModelingOperators/Public/SpaceDeformerOps/BendMeshOp.h b/Engine/Plugins/Experimental/MeshModelingToolset/Source/ModelingOperators/Public/SpaceDeformerOps/BendMeshOp.h
new file mode 100644
--- /dev/null
+++ b/Engine/Plugins/Experimental/MeshModelingToolset/Source/ModelingOperators/Public/SpaceDeformerOps/BendMeshOp.h
@@ -0,0 +1,17 @@
+// Copyright 1998-2019 Epic Games, Inc. All Rights Reserved.
+
+#pragma once
+
+#include "CoreMinimal.h"
+#include "MeshSpaceDeformerOp.h"
+
+
+class MODELINGOPERATORS_API FBendMeshOp : public FMeshSpaceDeformerOp
+{
+public:
+	virtual ~FBendMeshOp() {}
+	FBendMeshOp() : FMeshSpaceDeformerOp(0.0, 180.0) {};
+	virtual void CalculateResult(FProgressCancel* Progress) override;
+
+protected:
+};
\ No newline at end of file
diff --git a/Engine/Plugins/Experimental/MeshModelingToolset/Source/ModelingOperators/Public/SpaceDeformerOps/FlareMeshOp.h b/Engine/Plugins/Experimental/MeshModelingToolset/Source/ModelingOperators/Public/SpaceDeformerOps/FlareMeshOp.h
new file mode 100644
--- /dev/null
+++ b/Engine/Plugins/Experimental/MeshModelingToolset/Source/ModelingOperators/Public/SpaceDeformerOps/FlareMeshOp.h
@@ -0,0 +1,20 @@
+// Copyright 1998-2019 Epic Games, Inc. All Rights Reserved.
+
+#pragma once
+
+#include "CoreMinimal.h"
+#include "MeshSpaceDeformerOp.h"
+
+
+
+class MODELINGOPERATORS_API FFlareMeshOp : public FMeshSpaceDeformerOp
+{
+public:
+	virtual ~FFlareMeshOp() {}
+	FFlareMeshOp() : FMeshSpaceDeformerOp(0.1, 2.0) {};
+	virtual void CalculateResult(FProgressCancel* Progress) override;
+
+
+protected:
+
+};
\ No newline at end of file
diff --git a/Engine/Plugins/Experimental/MeshModelingToolset/Source/ModelingOperators/Public/SpaceDeformerOps/MeshSpaceDeformerOp.h b/Engine/Plugins/Experimental/MeshModelingToolset/Source/ModelingOperators/Public/SpaceDeformerOps/MeshSpaceDeformerOp.h
new file mode 100644
--- /dev/null
+++ b/Engine/Plugins/Experimental/MeshModelingToolset/Source/ModelingOperators/Public/SpaceDeformerOps/MeshSpaceDeformerOp.h
@@ -0,0 +1,74 @@
+// Copyright 1998-2019 Epic Games, Inc. All Rights Reserved.
+
+#pragma once
+
+#include "CoreMinimal.h"
+#include "Util/ProgressCancel.h"
+#include "ModelingOperators.h"
+
+class MODELINGOPERATORS_API FMeshSpaceDeformerOp : public FDynamicMeshOperator
+{
+public:
+	FMeshSpaceDeformerOp(double Min, double Max) : ModifierMin{ Min }, ModifierMax{ Max }{};
+	virtual ~FMeshSpaceDeformerOp() {}
+
+	FDynamicMesh3* TargetMesh;
+
+	// set ability on protected transform.
+
+	//Converts the given modifier percent and returns the value between -1 and 1
+	virtual double GetModifierValue() { return FMath::Lerp(0.01 * ModifierPercent, ModifierMin, ModifierMax); }
+
+	/**
+	* \brief Updates the pointer to the target mesh
+	*/
+	virtual void UpdateMesh(FDynamicMesh3* Mesh);
+
+
+	/**
+	* \brief Updates spatial data i.e. origin, orientation and other pertaining inputs needed by the operators
+	* \param ObjectSpaceToOpSpaceTransform Rotation matrix to rotate from object space to the space expected by the operator
+	* \param ObjectSpaceOrigin The location of the handle's origin, transformed from world space into the local coordinate system of the mesh
+	* \param AxesHalfExtents Each element of this vector is a scalar corresponding to the distance from the ObjectSpaceOrigin (centroid) to the farthest vertex in that respective direction
+	* \param LowerBounds The negative interval (only negative numbers) representing the range of effect in space this operator will have where -1 will utilize the entire half extent
+	* \param UpperBounds The positive interval (only positive numbers) representing the range of effect in space this operator will have where 1 will utilize the entire half extent
+	* \param ModifierValue Percent input from the tool (between -100 and 100), as each operator will have a min/max value i.e. the angle of curvature, twist or scale. This is the percent used to interpolate between the min/max value
+	*/
+	virtual void UpdateAxisData(const FMatrix3d& ObjectSpaceToOpSpaceTransform, const FVector3d& ObjectSpaceOrigin, const FVector3d& AxesHalfExtents, double LowerBounds, double UpperBounds, double ModifierValue);
+
+	//Array of the original positions of each vertex of the mesh
+	TArray<FVector3d> OriginalPositions;
+
+	//Origin of the transformation, given in the local coordinate system of the mesh
+	FVector3d AxisOriginObjectSpace;
+	
+	//Each element of this vector is a scalar corresponding to the distance from the ObjectSpaceOrigin (centroid) to the farthest vertex in that respective direction
+	FVector3d AxesHalfLengths;
+
+	//Rotation matrix to rotate from object space to the space expected by the operator
+	FMatrix3d ObjectSpaceToOpSpace;
+
+	//Precomputed inversion from the provided ObjectSpaceToOpSpace, used to convert back to the object's coordinate system
+	FMatrix3d OpSpaceToObjectSpace;
+
+	//Percent input from the tool (between -100 and 100), as each operator will have a min/max value i.e. the angle of curvature, twist or scale. This is the percent used to interpolate between the min/max value
+	double ModifierPercent;
+
+	//The minimum value provided by the constructor so that each operator can set it's own interpolation min and max
+	double ModifierMin{};
+
+	//The minimum value provided by the constructor so that each operator can set it's own interpolation min and max
+	double ModifierMax{};
+
+	//The interval usually from [-1,0] upon which this operator will work. This corresponds to the region of space which this operator affects.
+	double LowerBoundsInterval;
+
+	//The interval usually from [0,1] upon which this operator will work. This corresponds to the region of space which this operator affects.
+	double UpperBoundsInterval;
+
+	virtual void CalculateResult(FProgressCancel* Progress) override;
+
+
+protected:
+
+};
\ No newline at end of file
diff --git a/Engine/Plugins/Experimental/MeshModelingToolset/Source/ModelingOperators/Public/SpaceDeformerOps/TwistMeshOp.h b/Engine/Plugins/Experimental/MeshModelingToolset/Source/ModelingOperators/Public/SpaceDeformerOps/TwistMeshOp.h
new file mode 100644
--- /dev/null
+++ b/Engine/Plugins/Experimental/MeshModelingToolset/Source/ModelingOperators/Public/SpaceDeformerOps/TwistMeshOp.h
@@ -0,0 +1,21 @@
+// Copyright 1998-2019 Epic Games, Inc. All Rights Reserved.
+
+#pragma once
+
+#include "CoreMinimal.h"
+#include "MeshSpaceDeformerOp.h"
+
+
+
+class MODELINGOPERATORS_API FTwistMeshOp : public FMeshSpaceDeformerOp
+{
+public:
+	virtual ~FTwistMeshOp() {}
+	FTwistMeshOp() : FMeshSpaceDeformerOp(0.0, 180.0) {}
+
+	virtual void CalculateResult(FProgressCancel* Progress) override;
+
+
+protected:
+
+};
\ No newline at end of file
diff --git a/Engine/Plugins/Experimental/MeshModelingToolset/Source/ModelingOperatorsEditorOnly/ModelingOperatorsEditorOnly.Build.cs b/Engine/Plugins/Experimental/MeshModelingToolset/Source/ModelingOperatorsEditorOnly/ModelingOperatorsEditorOnly.Build.cs
new file mode 100644
--- /dev/null
+++ b/Engine/Plugins/Experimental/MeshModelingToolset/Source/ModelingOperatorsEditorOnly/ModelingOperatorsEditorOnly.Build.cs
@@ -0,0 +1,65 @@
+// Copyright 1998-2019 Epic Games, Inc. All Rights Reserved.
+
+using UnrealBuildTool;
+
+public class ModelingOperatorsEditorOnly : ModuleRules
+{
+	public ModelingOperatorsEditorOnly(ReadOnlyTargetRules Target) : base(Target)
+	{
+		PCHUsage = ModuleRules.PCHUsageMode.UseExplicitOrSharedPCHs;
+		
+		PublicIncludePaths.AddRange(
+			new string[] {
+				// ... add public include paths required here ...
+			}
+			);
+				
+		
+		PrivateIncludePaths.AddRange(
+			new string[] {
+				// ... add other private include paths required here ...
+			}
+			);
+			
+		
+		PublicDependencyModuleNames.AddRange(
+			new string[]
+			{
+				"Core",
+                "InteractiveToolsFramework",
+                "MeshDescription",
+				"StaticMeshDescription",
+				"ModelingOperators",
+				"GeometricObjects",
+				"DynamicMesh",
+				"MeshConversion",
+				"MeshSolverUtilities",
+				"GeometryAlgorithms", // required for constrained Delaunay triangulation
+                "ProxyLODMeshReduction", // currently required to be public due to IVoxelBasedCSG API
+				// ... add other public dependencies that you statically link with here ...
+			}
+            );
+			
+		
+		PrivateDependencyModuleNames.AddRange(
+			new string[]
+			{
+				"CoreUObject",
+				"Engine",
+
+				"MeshBuilder",
+                "MeshDescriptionOperations",
+				"MeshUtilitiesCommon",  
+				"MeshReductionInterface", // for UE4 standard simplification 
+			}
+			);
+		
+		
+		DynamicallyLoadedModuleNames.AddRange(
+			new string[]
+			{
+				// ... add any modules that your module loads dynamically here ...
+			}
+			);
+	}
+}
diff --git a/Engine/Plugins/Experimental/MeshModelingToolset/Source/ModelingOperatorsEditorOnly/Private/CompositionOps/VoxelBooleanMeshesOp.cpp b/Engine/Plugins/Experimental/MeshModelingToolset/Source/ModelingOperatorsEditorOnly/Private/CompositionOps/VoxelBooleanMeshesOp.cpp
new file mode 100644
--- /dev/null
+++ b/Engine/Plugins/Experimental/MeshModelingToolset/Source/ModelingOperatorsEditorOnly/Private/CompositionOps/VoxelBooleanMeshesOp.cpp
@@ -0,0 +1,136 @@
+// Copyright 1998-2019 Epic Games, Inc. All Rights Reserved.
+
+#pragma once
+
+#include "CompositionOps/VoxelBooleanMeshesOp.h"
+
+#include "Engine/StaticMesh.h"
+
+#include "MeshDescription.h"
+#include "StaticMeshAttributes.h"
+#include "MeshDescriptionToDynamicMesh.h"
+#include "MeshSimplification.h"
+#include "MeshConstraintsUtil.h"
+
+void FVoxelBooleanMeshesOp::CalculateResult(FProgressCancel* Progress)
+{
+	FMeshDescription ResultMeshDescription;
+	FStaticMeshAttributes Attributes(ResultMeshDescription);
+	Attributes.Register();
+
+	if (VoxelCount < 16)
+	{
+		VoxelCount = 16;
+	}
+
+	// Use the world space bounding box of each mesh to compute the voxel size
+	VoxelSizeD = ComputeVoxelSize();
+
+	// give this an absolute min since the user might scale both objects to zero..
+	VoxelSizeD = FMath::Max(VoxelSizeD, 0.001);
+
+	// Create BooleanTool and Boolean the meshes.
+	TUniquePtr<IVoxelBasedCSG> VoxelCSGTool = IVoxelBasedCSG::CreateCSGTool(VoxelSizeD);
+	FVector MergedOrigin;
+
+	IVoxelBasedCSG::FPlacedMesh& A = (*InputMeshArray)[0];
+	IVoxelBasedCSG::FPlacedMesh& B = (*InputMeshArray)[1];
+
+	switch (Operation)
+	{
+	case EBooleanOperation::DifferenceAB:
+		MergedOrigin = VoxelCSGTool->ComputeDifference(A, B, ResultMeshDescription, AdaptivityD, IsoSurfaceD);
+		break;
+	case EBooleanOperation::DifferenceBA:
+		MergedOrigin = VoxelCSGTool->ComputeDifference(B, A, ResultMeshDescription, AdaptivityD, IsoSurfaceD);
+		break;
+	case EBooleanOperation::Intersect:
+		MergedOrigin = VoxelCSGTool->ComputeIntersection(A, B, ResultMeshDescription, AdaptivityD, IsoSurfaceD);
+		break;
+	case EBooleanOperation::Union:
+		MergedOrigin = VoxelCSGTool->ComputeUnion(A, B, ResultMeshDescription, AdaptivityD, IsoSurfaceD);
+		break;
+	default:
+		check(0);
+	}
+
+	ResultTransform = FTransform3d(MergedOrigin);
+
+	if (Progress->Cancelled())
+	{
+		return;
+	}
+
+	// Convert to dynamic mesh
+	FMeshDescriptionToDynamicMesh Converter;
+	Converter.bPrintDebugMessages = true;
+	Converter.Convert(&ResultMeshDescription, *ResultMesh);
+
+	if (Progress->Cancelled())
+	{
+		return;
+	}
+
+	if (bAutoSimplify)
+	{
+		FQEMSimplification Reducer(ResultMesh.Get());
+		FMeshConstraints constraints;
+		FMeshConstraintsUtil::ConstrainAllSeams(constraints, *ResultMesh, true, false);
+		Reducer.SetExternalConstraints(&constraints);
+		Reducer.Progress = Progress;
+
+		const double MaxDisplacementSqr = 3. * VoxelSizeD * VoxelSizeD;
+
+		Reducer.SimplifyToMaxError(MaxDisplacementSqr);
+
+	}
+}
+
+float FVoxelBooleanMeshesOp::ComputeVoxelSize() const 
+{ 
+	const TArray<IVoxelBasedCSG::FPlacedMesh>& PlacedMeshes = *(InputMeshArray);
+
+	float Size = 0.f;
+	for (int32 i = 0; i < PlacedMeshes.Num(); ++i)
+	{
+		//Bounding box
+		FVector BBoxMin(FLT_MAX, FLT_MAX, FLT_MAX);
+		FVector BBoxMax(-FLT_MAX, -FLT_MAX, -FLT_MAX);
+
+		// Use the scale define by the transform.  We don't care about actual placement
+		// in the world for this.
+		FVector Scale = PlacedMeshes[i].Transform.GetScale3D();
+		const FMeshDescription&  MeshDescription = *PlacedMeshes[i].Mesh;
+
+		TVertexAttributesConstRef<FVector> VertexPositions = MeshDescription.VertexAttributes().GetAttributesRef<FVector>(MeshAttribute::Vertex::Position);
+		for (const FVertexID VertexID : MeshDescription.Vertices().GetElementIDs())
+		{
+			const FVector Pos = VertexPositions[VertexID];
+            
+			BBoxMin.X = FMath::Min(BBoxMin.X, Pos.X);
+			BBoxMin.Y = FMath::Min(BBoxMin.Y, Pos.Y);
+			BBoxMin.Z = FMath::Min(BBoxMin.Z, Pos.Z);
+
+			BBoxMax.X = FMath::Max(BBoxMax.X, Pos.X);
+			BBoxMax.Y = FMath::Max(BBoxMax.Y, Pos.Y);
+			BBoxMax.Z = FMath::Max(BBoxMax.Z, Pos.Z);
+		}
+
+		// The size of the BBox in each direction
+		FVector Extents(BBoxMax.X - BBoxMin.X, BBoxMax.Y - BBoxMin.Y, BBoxMax.Z - BBoxMin.Z);
+		
+		// Scale with the local space scale.
+		Extents.X = Extents.X * Scale.X;
+		Extents.Y = Extents.Y * Scale.Y;
+		Extents.Z = Extents.Z * Scale.Z;
+
+		float MajorAxisSize = FMath::Max3(Extents.X, Extents.Y, Extents.Z);
+
+
+		Size = FMath::Max(MajorAxisSize / VoxelCount, Size);
+
+	}
+	
+	return Size;
+
+}
diff --git a/Engine/Plugins/Experimental/MeshModelingToolset/Source/ModelingOperatorsEditorOnly/Private/CompositionOps/VoxelMergeMeshesOp.cpp b/Engine/Plugins/Experimental/MeshModelingToolset/Source/ModelingOperatorsEditorOnly/Private/CompositionOps/VoxelMergeMeshesOp.cpp
new file mode 100644
--- /dev/null
+++ b/Engine/Plugins/Experimental/MeshModelingToolset/Source/ModelingOperatorsEditorOnly/Private/CompositionOps/VoxelMergeMeshesOp.cpp
@@ -0,0 +1,110 @@
+// Copyright 1998-2019 Epic Games, Inc. All Rights Reserved.
+
+#pragma once
+
+#include "CompositionOps/VoxelMergeMeshesOp.h"
+
+#include "Engine/StaticMesh.h"
+
+#include "MeshDescription.h"
+#include "StaticMeshAttributes.h"
+#include "MeshDescriptionToDynamicMesh.h"
+#include "MeshSimplification.h"
+#include "MeshConstraintsUtil.h"
+
+
+void FVoxelMergeMeshesOp::CalculateResult(FProgressCancel* Progress)
+{
+	FMeshDescription MergedMeshesDescription;
+	FStaticMeshAttributes Attributes(MergedMeshesDescription);
+	Attributes.Register();
+
+	// Use the world space bounding box of each mesh to compute the voxel size
+	VoxelSizeD = ComputeVoxelSize();
+
+	// give this an absolute min since the user might scale both objects to zero..
+	VoxelSizeD = FMath::Max(VoxelSizeD, 0.001);
+
+	// Create CSGTool and merge the meshes.
+	TUniquePtr<IVoxelBasedCSG> VoxelCSGTool = IVoxelBasedCSG::CreateCSGTool(VoxelSizeD);
+	FVector MergedOrigin = VoxelCSGTool->ComputeUnion(*InputMeshArray, MergedMeshesDescription, AdaptivityD, IsoSurfaceD);
+	ResultTransform = FTransform3d(MergedOrigin);
+
+	if (Progress->Cancelled())
+	{
+		return;
+	}
+
+	// Convert to dynamic mesh
+	FMeshDescriptionToDynamicMesh Converter;
+	Converter.bPrintDebugMessages = true;
+	Converter.Convert(&MergedMeshesDescription, *ResultMesh);
+
+	if (Progress->Cancelled())
+	{
+		return;
+	}
+
+	if (bAutoSimplify)
+	{
+		FQEMSimplification Reducer(ResultMesh.Get());
+		FMeshConstraints constraints;
+		FMeshConstraintsUtil::ConstrainAllSeams(constraints, *ResultMesh, true, false);
+		Reducer.SetExternalConstraints(&constraints);
+		Reducer.Progress = Progress;
+
+		const double MaxDisplacementSqr = 3. *VoxelSizeD * VoxelSizeD;
+
+		Reducer.SimplifyToMaxError(MaxDisplacementSqr);
+	}
+}
+
+
+float FVoxelMergeMeshesOp::ComputeVoxelSize() const
+{
+	const TArray<IVoxelBasedCSG::FPlacedMesh>& PlacedMeshes = *(InputMeshArray);
+
+	float Size = 0.f;
+	for (int32 i = 0; i < PlacedMeshes.Num(); ++i)
+	{
+		//Bounding box
+		FVector BBoxMin(FLT_MAX, FLT_MAX, FLT_MAX);
+		FVector BBoxMax(-FLT_MAX, -FLT_MAX, -FLT_MAX);
+
+		// Use the scale define by the transform.  We don't care about actual placement
+		// in the world for this.
+		FVector Scale = PlacedMeshes[i].Transform.GetScale3D();
+		const FMeshDescription&  MeshDescription = *PlacedMeshes[i].Mesh;
+
+		TVertexAttributesConstRef<FVector> VertexPositions = MeshDescription.VertexAttributes().GetAttributesRef<FVector>(MeshAttribute::Vertex::Position);
+		for (const FVertexID VertexID : MeshDescription.Vertices().GetElementIDs())
+		{
+			const FVector Pos = VertexPositions[VertexID];
+
+			BBoxMin.X = FMath::Min(BBoxMin.X, Pos.X);
+			BBoxMin.Y = FMath::Min(BBoxMin.Y, Pos.Y);
+			BBoxMin.Z = FMath::Min(BBoxMin.Z, Pos.Z);
+
+			BBoxMax.X = FMath::Max(BBoxMax.X, Pos.X);
+			BBoxMax.Y = FMath::Max(BBoxMax.Y, Pos.Y);
+			BBoxMax.Z = FMath::Max(BBoxMax.Z, Pos.Z);
+		}
+
+		// The size of the BBox in each direction
+		FVector Extents(BBoxMax.X - BBoxMin.X, BBoxMax.Y - BBoxMin.Y, BBoxMax.Z - BBoxMin.Z);
+
+		// Scale with the local space scale.
+		Extents.X = Extents.X * Scale.X;
+		Extents.Y = Extents.Y * Scale.Y;
+		Extents.Z = Extents.Z * Scale.Z;
+
+		float MajorAxisSize = FMath::Max3(Extents.X, Extents.Y, Extents.Z);
+
+
+		Size = FMath::Max(MajorAxisSize / VoxelCount, Size);
+
+	}
+
+	return Size;
+
+}
diff --git a/Engine/Plugins/Experimental/MeshModelingToolset/Source/ModelingOperatorsEditorOnly/Private/CuttingOps/EmbedPolygonsOp.cpp b/Engine/Plugins/Experimental/MeshModelingToolset/Source/ModelingOperatorsEditorOnly/Private/CuttingOps/EmbedPolygonsOp.cpp
new file mode 100644
--- /dev/null
+++ b/Engine/Plugins/Experimental/MeshModelingToolset/Source/ModelingOperatorsEditorOnly/Private/CuttingOps/EmbedPolygonsOp.cpp
@@ -0,0 +1,124 @@
+// Copyright 1998-2019 Epic Games, Inc. All Rights Reserved.
+
+#pragma once
+
+#include "CuttingOps/EmbedPolygonsOp.h"
+
+#include "DynamicMeshEditor.h"
+#include "Selections/MeshFaceSelection.h"
+#include "MeshQueries.h"
+#include "Operations/EmbedSurfacePath.h"
+#include "Operations/SimpleHoleFiller.h"
+
+#include "Engine/StaticMesh.h"
+
+#include "MeshDescriptionToDynamicMesh.h"
+#include "MeshSimplification.h"
+#include "MeshConstraintsUtil.h"
+#include "Operations/MeshPlaneCut.h"
+#include "ConstrainedDelaunay2.h"
+
+
+void FEmbedPolygonsOp::SetTransform(const FTransform& Transform) {
+	ResultTransform = (FTransform3d)Transform;
+}
+
+void FEmbedPolygonsOp::CalculateResult(FProgressCancel* Progress)
+{
+	if (Progress->Cancelled())
+	{
+		return;
+	}
+	ResultMesh->Copy(*OriginalMesh, true, true, true, !bDiscardAttributes);
+
+
+
+	
+	FVector3d p = LocalPlaneOrigin + LocalPlaneNormal * 1000;
+	FVector3d d = -LocalPlaneNormal;
+	FFrame3d Frame(p, d);
+	
+
+	FPolygon2d Polygon = GetPolygon();
+	
+	TArray<int> PathVertIDs1, PathVertCorrespond1;
+	
+	TArray<TPair<float, int>> SortedHitTriangles;
+	TMeshQueries<FDynamicMesh3>::FindHitTriangles_LinearSearch(*ResultMesh, FRay3d(Frame.FromPlaneUV(Polygon[0]), Frame.Z()), SortedHitTriangles);
+	if (SortedHitTriangles.Num() < 1 || (SortedHitTriangles.Num() < 2 && Operation == EEmbeddedPolygonOpMethod::CutThrough))
+	{
+		// didn't hit the mesh 
+		return;
+	}
+
+	auto CutHole = [](FDynamicMesh3& Mesh, const FFrame3d& F, int TriStart, const FPolygon2d& PolygonArg, TArray<int>& PathVertIDs, TArray<int>& PathVertCorrespond)
+	{
+		FMeshFaceSelection Selection(&Mesh);
+		bool bDidEmbed = EmbedProjectedPath(&Mesh, TriStart, F, PolygonArg.GetVertices(), PathVertIDs, PathVertCorrespond, true, &Selection);
+		if (!bDidEmbed)
+		{
+			return false;
+		}
+		// slow debugging check, TODO: remove
+		for (int VID : PathVertIDs)
+		{
+			ensure(Mesh.IsVertex(VID));
+		}
+
+		FDynamicMeshEditor MeshEditor(&Mesh);
+		bool bDidRemove = MeshEditor.RemoveTriangles(Selection.AsArray(), true);
+		if (!bDidRemove)
+		{
+			return false;
+		}
+		for (int VID : PathVertIDs)
+		{
+			// check if the triangle removal nuked the vertices on the boundary; TODO: don't 'ensure' if we need to keep this, but probably some earlier code should handle this ...
+			if (!ensure(Mesh.IsVertex(VID)))
+			{
+				return false;
+			}
+		}
+
+		return true;
+	};
+
+	bool bCutSide1 = CutHole(*ResultMesh, Frame, SortedHitTriangles[0].Value, Polygon, PathVertIDs1, PathVertCorrespond1);
+	if (!bCutSide1 || PathVertIDs1.Num() < 2)
+	{
+		return;
+	}
+
+	if (Operation == EEmbeddedPolygonOpMethod::CutThrough)
+	{
+		TArray<int> PathVertIDs2, PathVertCorrespond2;
+		bool bCutSide2 = CutHole(*ResultMesh, Frame, SortedHitTriangles[1].Value, Polygon, PathVertIDs2, PathVertCorrespond2);
+		if (!bCutSide2 || PathVertIDs2.Num() < 2)
+		{
+			return;
+		}
+		FDynamicMeshEditor MeshEditor(ResultMesh.Get());
+		FDynamicMeshEditResult ResultOut;
+		bool bStitched = MeshEditor.StitchSparselyCorrespondedVertexLoops(PathVertIDs2, PathVertCorrespond2, PathVertIDs1, PathVertCorrespond1, ResultOut);
+		// TODO: set triangle normals and uvs if mesh hasattributes
+	}
+	else if (Operation == EEmbeddedPolygonOpMethod::CutAndFill)
+	{
+		FDynamicMeshEditor MeshEditor(ResultMesh.Get());
+		FEdgeLoop Loop(ResultMesh.Get());
+		Loop.InitializeFromVertices(PathVertIDs1);
+
+		
+		FSimpleHoleFiller Filler(ResultMesh.Get(), Loop);
+		int GID = ResultMesh->AllocateTriangleGroup();
+		if (Filler.Fill(GID))
+		{
+			if (ResultMesh->HasAttributes())
+			{
+				MeshEditor.SetTriangleNormals(Filler.NewTriangles, (FVector3f)-Frame.Z()); // TODO: should we use the best fit plane instead of the projection plane for the normal? ... probably.
+				float UVScaleFactor = 1; // TODO: this should be a parameter
+				MeshEditor.SetTriangleUVsFromProjection(Filler.NewTriangles, Frame, UVScaleFactor);
+			}
+		}
+	}
+}
\ No newline at end of file
diff --git a/Engine/Plugins/Experimental/MeshModelingToolset/Source/ModelingOperatorsEditorOnly/Private/ParameterizationOps/ParameterizeMeshOp.cpp b/Engine/Plugins/Experimental/MeshModelingToolset/Source/ModelingOperatorsEditorOnly/Private/ParameterizationOps/ParameterizeMeshOp.cpp
new file mode 100644
--- /dev/null
+++ b/Engine/Plugins/Experimental/MeshModelingToolset/Source/ModelingOperatorsEditorOnly/Private/ParameterizationOps/ParameterizeMeshOp.cpp
@@ -0,0 +1,238 @@
+// Copyright 1998-2019 Epic Games, Inc. All Rights Reserved.
+
+#pragma once
+
+#include "ParameterizationOps/ParameterizeMeshOp.h"
+
+#include "DynamicMeshAttributeSet.h"
+#include "MeshDescriptionToDynamicMesh.h"
+#include "ProxyLODParameterization.h"
+
+
+
+
+FParameterizeMeshOp::FLinearMesh::FLinearMesh(const FDynamicMesh3& Mesh)
+{
+
+	TArray<FVector>& Positions = this->VertexBuffer;
+
+	// Temporary maps used during construction.
+
+	TArray<int32> TriFromID;
+	TArray<int32> TriToID;
+	TArray<int32> VertFromID;
+
+	// Compute the mapping from triangle ID to triangle number
+	{
+		const int32 MaxTriID = Mesh.MaxTriangleID();
+		const int32 NumTris = Mesh.TriangleCount();
+
+		// reserve space and add elements 
+		TriFromID.Empty(MaxTriID);
+		TriFromID.AddUninitialized(MaxTriID);
+
+		// reserve space
+		TriToID.Empty(NumTris);
+
+		int32 count = 0;
+		for (int TriID : Mesh.TriangleIndicesItr())
+		{
+			TriToID.Add(TriID);
+			TriFromID[TriID] = count;
+			count++;
+		}
+	}
+
+	// Compute the mapping from vertex ID to vertex number
+	{
+		const int32 MaxVertID = Mesh.MaxVertexID();
+		const int32 NumVerts = Mesh.VertexCount();
+
+		// reserve space and add elements
+		VertFromID.Empty(MaxVertID);
+		VertFromID.AddUninitialized(MaxVertID);
+
+		// reserve space
+		VertToID.Empty(NumVerts);
+
+		int32 count = 0;
+		for (int VtxID : Mesh.VertexIndicesItr())
+		{
+			VertToID.Add(VtxID);
+			VertFromID[VtxID] = count;
+			count++;
+		}
+	}
+
+	// Fill the vertex buffer
+	{
+		int32 NumVerts = Mesh.VertexCount();
+		Positions.Empty(NumVerts);
+
+		for (const auto& Vertex : Mesh.VerticesItr())
+		{
+			FVector Pos(Vertex.X, Vertex.Y, Vertex.Z);
+			Positions.Add(Pos);
+		}
+	}
+
+	const int32 NumTris = Mesh.TriangleCount();
+
+	// Fill the index buffer
+	{
+		IndexBuffer.Empty(NumTris * 3);
+		for (const auto& Tri : Mesh.TrianglesItr())
+		{
+			for (int i = 0; i < 3; ++i)
+			{
+				int VtxID = Tri[i];
+				int32 RemapVtx = VertFromID[VtxID];
+				IndexBuffer.Add(RemapVtx);
+			}
+		}
+
+	}
+
+
+	// For each edge on each triangle.
+	AdjacencyBuffer.Empty(NumTris * 3);
+
+	for (int TriID : Mesh.TriangleIndicesItr())
+	{
+
+		// NB: this maybe the same as GetNeighborTriangles()
+		for (int i = 0; i < 3; ++i)
+		{
+			int32 EdgesID = Mesh.GetTriEdge(TriID, i);
+			FIndex4i Edge = Mesh.GetEdge(EdgesID);
+			int32 OtherTriID = (Edge[2] == TriID) ? Edge[3] : Edge[2];
+
+			int32 RemapOtherID = (OtherTriID != FDynamicMesh3::InvalidID) ? TriFromID[OtherTriID] : FDynamicMesh3::InvalidID;
+
+			AdjacencyBuffer.Add(RemapOtherID);
+		}
+	}
+
+}
+
+bool FParameterizeMeshOp::ComputeUVs(FDynamicMesh3& Mesh, TFunction<bool(float)>& Interrupter)
+{
+	// Convert to a dense form.
+	FLinearMesh LinearMesh(Mesh);
+
+	// Data to be populated by the UV generation tool
+	TArray<FVector2D> UVVertexBuffer;
+	TArray<int32>     UVIndexBuffer;
+	TArray<int32>     VertexRemapArray;
+
+
+	float MaxStretch     = Stretch;
+	int32 MaxChartNumber = NumCharts;
+
+	TUniquePtr<IProxyLODParameterization> ParameterizationTool = IProxyLODParameterization::CreateTool();
+	bool bSuccess = ParameterizationTool->GenerateUVs(Width, Height, Gutter, LinearMesh.VertexBuffer, LinearMesh.IndexBuffer, LinearMesh.AdjacencyBuffer, Interrupter, UVVertexBuffer, UVIndexBuffer, VertexRemapArray, MaxStretch, MaxChartNumber);
+
+
+	// Add the UVs to the FDynamicMesh
+	if (bSuccess)
+	{
+
+		const bool bHasAttributes = Mesh.HasAttributes();
+
+		if (bHasAttributes)
+		{
+			FDynamicMeshAttributeSet* Attributes = Mesh.Attributes();
+			Attributes->PrimaryUV()->ClearElements(); // delete existing UVs
+		}
+		else
+		{
+			// Add attrs for UVS
+			Mesh.EnableAttributes();
+		}
+
+		FDynamicMeshUVOverlay* UVOverlay = Mesh.Attributes()->PrimaryUV();
+
+		// This mesh shouldn't already have UVs.
+
+		checkSlow(UVOverlay->ElementCount() == 0);
+
+		// Add the UVs to the overlay
+		int32 NumUVs = UVVertexBuffer.Num();
+		for (int32 i = 0; i < NumUVs; ++i)
+		{
+			const FVector2D UV = UVVertexBuffer[i];
+
+			// The associated VertID in the dynamic mesh
+			const int32 VertOffset = VertexRemapArray[i];
+			const int32 VertID = LinearMesh.VertToID[VertOffset];
+
+			// add the UV to the mesh overlay
+			const int32 NewID = UVOverlay->AppendElement(UV, VertID);
+
+		}
+
+		int32 NumUVTris = UVIndexBuffer.Num() / 3;
+		for (int32 i = 0; i < NumUVTris; ++i)
+		{
+			int32 t = i * 3;
+			// The triangle in UV space
+			FIndex3i TriUV(UVIndexBuffer[t], UVIndexBuffer[t + 1], UVIndexBuffer[t + 2]);
+
+			// the triangle in terms of the VertIDs in the DynamicMesh
+			FIndex3i TriVertIDs;
+			for (int c = 0; c < 3; ++c)
+			{
+				// the offset for this vertex in the LinearMesh
+				int32 Offset = VertexRemapArray[TriUV[c]];
+
+				int32 VertId = LinearMesh.VertToID[Offset];
+
+				TriVertIDs[c] = VertId;
+			}
+
+			// NB: this could be slow.. 
+			int32 TriID = Mesh.FindTriangle(TriVertIDs[0], TriVertIDs[1], TriVertIDs[2]);
+
+			checkSlow(TriID != FDynamicMesh3::InvalidID);
+
+			// add the triangle to the overlay
+			UVOverlay->SetTriangle(TriID, TriUV);
+		}
+
+	}
+
+	return bSuccess;
+}
+
+
+void FParameterizeMeshOp::CalculateResult(FProgressCancel* Progress)
+{
+
+	if (!InputMesh.IsValid())
+	{
+		return;
+	}
+	// Need access to the source mesh:
+
+	const FMeshDescription* MeshDescription = InputMesh.Get();
+
+	// Convert to FDynamic Mesh
+	{
+		FMeshDescriptionToDynamicMesh Converter;
+		Converter.bPrintDebugMessages = true;
+		ResultMesh->Clear();
+		Converter.Convert(MeshDescription, *ResultMesh);
+	}
+
+	if (Progress->Cancelled())
+	{
+		return;
+	}
+
+	// The UV atlas callback uses a float progress-based interrupter. 
+	TFunction<bool(float)> Iterrupter = [Progress](float)->bool {return !Progress->Cancelled(); };
+
+
+	bool bSuccess = ComputeUVs(*ResultMesh,  Iterrupter);
+
+}
\ No newline at end of file
diff --git a/Engine/Plugins/Experimental/MeshModelingToolset/Source/ModelingOperatorsEditorOnly/Public/CompositionOps/VoxelBooleanMeshesOp.h b/Engine/Plugins/Experimental/MeshModelingToolset/Source/ModelingOperatorsEditorOnly/Public/CompositionOps/VoxelBooleanMeshesOp.h
new file mode 100644
--- /dev/null
+++ b/Engine/Plugins/Experimental/MeshModelingToolset/Source/ModelingOperatorsEditorOnly/Public/CompositionOps/VoxelBooleanMeshesOp.h
@@ -0,0 +1,48 @@
+// Copyright 1998-2019 Epic Games, Inc. All Rights Reserved.
+
+#pragma once
+
+#include "CoreMinimal.h"
+#include "Util/ProgressCancel.h"
+#include "ModelingOperators.h"
+#include "ProxyLODVolume.h"
+
+
+
+
+class MODELINGOPERATORSEDITORONLY_API FVoxelBooleanMeshesOp : public FDynamicMeshOperator
+{
+public:
+	virtual ~FVoxelBooleanMeshesOp() {}
+
+	enum class EBooleanOperation
+	{
+		DifferenceAB = 0,
+		DifferenceBA = 1,
+		Intersect = 2,
+		Union = 3
+	};
+
+	// inputs
+	TSharedPtr<TArray<IVoxelBasedCSG::FPlacedMesh>> InputMeshArray;
+	int32 VoxelCount = 128;
+	double VoxelSizeD = 1.0;
+	double AdaptivityD = 0;
+	double IsoSurfaceD = 0;
+	bool   bAutoSimplify = false;
+	EBooleanOperation Operation = EBooleanOperation::DifferenceBA;
+
+	//
+	// FDynamicMeshOperator implementation
+	// 
+
+	virtual void CalculateResult(FProgressCancel* Progress) override;
+
+private:
+
+	// compute the voxel size based on the voxel count and the input geometry
+	float ComputeVoxelSize() const;
+
+};
+
+
diff --git a/Engine/Plugins/Experimental/MeshModelingToolset/Source/ModelingOperatorsEditorOnly/Public/CompositionOps/VoxelMergeMeshesOp.h b/Engine/Plugins/Experimental/MeshModelingToolset/Source/ModelingOperatorsEditorOnly/Public/CompositionOps/VoxelMergeMeshesOp.h
new file mode 100644
--- /dev/null
+++ b/Engine/Plugins/Experimental/MeshModelingToolset/Source/ModelingOperatorsEditorOnly/Public/CompositionOps/VoxelMergeMeshesOp.h
@@ -0,0 +1,42 @@
+// Copyright 1998-2019 Epic Games, Inc. All Rights Reserved.
+
+#pragma once
+
+#include "CoreMinimal.h"
+#include "Util/ProgressCancel.h"
+#include "ModelingOperators.h"
+#include "ProxyLODVolume.h"
+
+
+
+
+
+class MODELINGOPERATORSEDITORONLY_API FVoxelMergeMeshesOp : public FDynamicMeshOperator
+{
+public:
+	virtual ~FVoxelMergeMeshesOp() {}
+
+	// inputs
+	TSharedPtr<TArray<IVoxelBasedCSG::FPlacedMesh>> InputMeshArray;
+	int32 VoxelCount = 128;
+	double VoxelSizeD = 1.0;
+	double AdaptivityD = 0;
+	double IsoSurfaceD = 0;
+	bool bAutoSimplify = false;
+
+	double FastCollapseVoxelSizeMultipler = 1.2;
+	int NumFastCollapsePasses = 10;
+
+	//
+	// FDynamicMeshOperator implementation
+	// 
+
+	virtual void CalculateResult(FProgressCancel* Progress) override;
+
+private:
+
+	// compute the voxel size based on the voxel count and the input geometry
+	float ComputeVoxelSize() const;
+};
+
+
diff --git a/Engine/Plugins/Experimental/MeshModelingToolset/Source/ModelingOperatorsEditorOnly/Public/CuttingOps/EmbedPolygonsOp.h b/Engine/Plugins/Experimental/MeshModelingToolset/Source/ModelingOperatorsEditorOnly/Public/CuttingOps/EmbedPolygonsOp.h
new file mode 100644
--- /dev/null
+++ b/Engine/Plugins/Experimental/MeshModelingToolset/Source/ModelingOperatorsEditorOnly/Public/CuttingOps/EmbedPolygonsOp.h
@@ -0,0 +1,57 @@
+// Copyright 1998-2019 Epic Games, Inc. All Rights Reserved.
+
+#pragma once
+
+#include "CoreMinimal.h"
+#include "Util/ProgressCancel.h"
+#include "ModelingOperators.h"
+#include "ProxyLODVolume.h"
+
+#include "Polygon2.h"
+
+#include "EmbedPolygonsOp.generated.h"
+
+
+UENUM()
+enum class EEmbeddedPolygonOpMethod : uint8
+{
+	CutAndFill,
+	CutThrough
+	//, Extrude  // TODO: extrude(/intrude?) would also be easy/natural to support here
+};
+
+
+class MODELINGOPERATORSEDITORONLY_API FEmbedPolygonsOp : public FDynamicMeshOperator
+{
+public:
+	virtual ~FEmbedPolygonsOp() {}
+
+	
+	// inputs
+	FVector LocalPlaneOrigin, LocalPlaneNormal;
+	float PolygonScale;
+
+	// TODO: stop hardcoding the polygon shape, switch to FGeneralPolygon2d
+	FPolygon2d GetPolygon()
+	{
+		return FPolygon2d::MakeCircle(PolygonScale, 20);
+	}
+
+	bool bDiscardAttributes;
+
+	EEmbeddedPolygonOpMethod Operation;
+
+	//float ExtrudeDistance; // TODO if we support extrude
+
+	TSharedPtr<FDynamicMesh3> OriginalMesh;
+
+	void SetTransform(const FTransform& Transform);
+
+	//
+	// FDynamicMeshOperator implementation
+	// 
+
+	virtual void CalculateResult(FProgressCancel* Progress) override;
+};
+
+
diff --git a/Engine/Plugins/Experimental/MeshModelingToolset/Source/ModelingOperatorsEditorOnly/Public/ParameterizationOps/ParameterizeMeshOp.h b/Engine/Plugins/Experimental/MeshModelingToolset/Source/ModelingOperatorsEditorOnly/Public/ParameterizationOps/ParameterizeMeshOp.h
new file mode 100644
--- /dev/null
+++ b/Engine/Plugins/Experimental/MeshModelingToolset/Source/ModelingOperatorsEditorOnly/Public/ParameterizationOps/ParameterizeMeshOp.h
@@ -0,0 +1,68 @@
+// Copyright 1998-2019 Epic Games, Inc. All Rights Reserved.
+
+#pragma once
+
+#include "CoreMinimal.h"
+#include "Util/ProgressCancel.h"
+#include "ModelingOperators.h"
+
+
+struct FMeshDescription;
+
+
+class MODELINGOPERATORSEDITORONLY_API FParameterizeMeshOp : public FDynamicMeshOperator
+{
+public:
+	virtual ~FParameterizeMeshOp() {}
+
+	//
+	// Inputs
+	// 
+
+	// source mesh
+	TSharedPtr<FMeshDescription> InputMesh;
+		
+	// UV generation parameters
+	float Stretch;
+	int32 NumCharts;
+
+	// Atlas Packing parameters
+	int32 Height = 512;
+	int32 Width = 512;
+	float Gutter = 2.5;
+
+	// set ability on protected transform.
+	void SetTransform(FTransform3d& XForm)
+	{
+		ResultTransform = XForm;
+	}
+
+	//
+	// FDynamicMeshOperator implementation
+	// 
+
+	virtual void CalculateResult(FProgressCancel* Progress) override;
+
+
+protected:
+
+	// dense index/vertex buffer based representation of the data needed for parameterization
+	struct FLinearMesh
+	{
+		FLinearMesh(const FDynamicMesh3& Mesh);
+
+		// Stripped down mesh
+		TArray<int32>   IndexBuffer;
+		TArray<FVector> VertexBuffer;
+
+		// Map from offset in the VertexBuffer to the VertexID in the FDynamicMesh
+		TArray<int32> VertToID;
+
+		// Adjacency[tri], Adjacency[tri+1], Adjacency[tri+2] 
+		// are the three tris adjacent to tri.
+		TArray<int32>   AdjacencyBuffer;
+		
+	};
+
+	bool ComputeUVs(FDynamicMesh3& InOutMesh, TFunction<bool(float)>& Interrupter);
+};
diff --git a/Engine/Plugins/Experimental/ModelingToolsEditorMode/Source/ModelingToolsEditorMode/Private/ModelingToolsEditorMode.cpp b/Engine/Plugins/Experimental/ModelingToolsEditorMode/Source/ModelingToolsEditorMode/Private/ModelingToolsEditorMode.cpp
new file mode 100644
--- /dev/null
+++ b/Engine/Plugins/Experimental/ModelingToolsEditorMode/Source/ModelingToolsEditorMode/Private/ModelingToolsEditorMode.cpp
@@ -0,0 +1,507 @@
+// Copyright 1998-2019 Epic Games, Inc. All Rights Reserved.
+
+#include "ModelingToolsEditorMode.h"
+#include "ModelingToolsEditorModeToolkit.h"
+#include "Toolkits/ToolkitManager.h"
+#include "Framework/Commands/UICommandList.h"
+#include "Framework/Application/SlateApplication.h"
+#include "Framework/MultiBox/MultiBoxBuilder.h"
+#include "EditorViewportClient.h"
+
+//#include "SingleClickTool.h"
+//#include "MeshSurfacePointTool.h"
+//#include "MeshVertexDragTool.h"
+#include "DynamicMeshSculptTool.h"
+#include "EditMeshPolygonsTool.h"
+#include "ConvertToPolygonsTool.h"
+#include "AddPrimitiveTool.h"
+#include "SmoothMeshTool.h"
+#include "RemeshMeshTool.h"
+#include "SimplifyMeshTool.h"
+#include "MeshInspectorTool.h"
+#include "WeldMeshEdgesTool.h"
+#include "DrawPolygonTool.h"
+#include "ShapeSprayTool.h"
+#include "MergeMeshesTool.h"
+#include "VoxelCSGMeshesTool.h"
+#include "PlaneCutTool.h"
+#include "PolygonOnMeshTool.h"
+#include "DisplaceMeshTool.h"
+#include "MeshSpaceDeformerTool.h"
+#include "EditNormalsTool.h"
+#include "AttributeEditorTool.h"
+#include "TransformMeshesTool.h"
+#include "MeshSelectionTool.h"
+#include "UVProjectionTool.h"
+
+#include "ParameterizeMeshTool.h"
+
+#include "EditorModeManager.h"
+
+// viewport interaction support
+#include "ViewportInteractor.h"
+#include "ActorViewportTransformable.h"
+#include "ViewportWorldInteraction.h"
+#include "IViewportInteractionModule.h"
+#include "LevelEditor.h"
+#include "IAssetViewport.h"
+#include "SLevelViewport.h"
+
+#include "ModelingToolsActions.h"
+#include "ModelingToolsManagerActions.h"
+
+#define LOCTEXT_NAMESPACE "FModelingToolsEditorMode"
+
+
+//#define ENABLE_DEBUG_PRINTING
+//#define TOOLED_ENABLE_VIEWPORT_INTERACTION
+
+const FEditorModeID FModelingToolsEditorMode::EM_ModelingToolsEditorModeId = TEXT("EM_ModelingToolsEditorMode");
+
+FModelingToolsEditorMode::FModelingToolsEditorMode()
+{
+	ToolsContext = nullptr;
+
+	UICommandList = MakeShareable(new FUICommandList);
+}
+
+FModelingToolsEditorMode::~FModelingToolsEditorMode()
+{
+	if (ToolsContext != nullptr)
+	{
+		ToolsContext->ShutdownContext();
+		ToolsContext = nullptr;
+	}
+}
+
+
+void FModelingToolsEditorMode::ActorSelectionChangeNotify()
+{
+}
+
+
+bool FModelingToolsEditorMode::ProcessEditDelete()
+{
+	return ToolsContext->ProcessEditDelete();
+}
+
+
+bool FModelingToolsEditorMode::AllowWidgetMove()
+{ 
+	return false; 
+}
+
+bool FModelingToolsEditorMode::ShouldDrawWidget() const
+{ 
+	// allow standard xform gizmo if we don't have an active tool
+	if (ToolsContext != nullptr)
+	{
+		return ToolsContext->ToolManager->HasAnyActiveTool() == false;
+	}
+	return true; 
+}
+
+bool FModelingToolsEditorMode::UsesTransformWidget() const
+{ 
+	return true; 
+}
+
+
+void FModelingToolsEditorMode::Tick(FEditorViewportClient* ViewportClient, float DeltaTime)
+{
+	FEdMode::Tick(ViewportClient, DeltaTime);
+
+	if (ToolsContext != nullptr)
+	{
+		ToolsContext->Tick(ViewportClient, DeltaTime);
+	}
+}
+
+
+
+
+void FModelingToolsEditorMode::Render(const FSceneView* View, FViewport* Viewport, FPrimitiveDrawInterface* PDI)
+{
+	FEdMode::Render(View, Viewport, PDI);
+
+	if (ToolsContext != nullptr)
+	{
+		ToolsContext->Render(View, Viewport, PDI);
+	}
+}
+
+
+
+
+
+bool FModelingToolsEditorMode::InputKey(FEditorViewportClient* ViewportClient, FViewport* Viewport, FKey Key, EInputEvent Event)
+{
+	// try hotkeys
+	if (Event != IE_Released)
+	{
+		if (UICommandList->ProcessCommandBindings(Key, FSlateApplication::Get().GetModifierKeys(), false/*Event == IE_Repeat*/))
+		{
+			return true;
+		}
+	}
+
+	bool bHandled = ToolsContext->InputKey(ViewportClient, Viewport, Key, Event);
+	if (bHandled == false)
+	{
+		bHandled = FEdMode::InputKey(ViewportClient, Viewport, Key, Event);
+	}
+	return bHandled;
+
+	//bool bHandled = FEdMode::InputKey(ViewportClient, Viewport, Key, Event);
+	//bHandled |= ToolsContext->InputKey(ViewportClient, Viewport, Key, Event);
+	//return bHandled;
+}
+
+
+bool FModelingToolsEditorMode::InputAxis(FEditorViewportClient* InViewportClient, FViewport* Viewport, int32 ControllerId, FKey Key, float Delta, float DeltaTime)
+{
+	// mouse axes: EKeys::MouseX, EKeys::MouseY, EKeys::MouseWheelAxis
+	return FEdMode::InputAxis(InViewportClient, Viewport, ControllerId, Key, Delta, DeltaTime);
+}
+
+
+
+
+
+bool FModelingToolsEditorMode::StartTracking(FEditorViewportClient* InViewportClient, FViewport* InViewport) 
+{ 
+	bool bHandled = FEdMode::StartTracking(InViewportClient, InViewport);
+#ifdef ENABLE_DEBUG_PRINTING
+	UE_LOG(LogTemp, Warning, TEXT("START TRACKING - base handled was %d"), (int)bHandled);
+#endif
+
+	bHandled |= ToolsContext->StartTracking(InViewportClient, InViewport);
+
+	return bHandled;
+}
+
+bool FModelingToolsEditorMode::CapturedMouseMove(FEditorViewportClient* InViewportClient, FViewport* InViewport, int32 InMouseX, int32 InMouseY)
+{
+	bool bHandled = ToolsContext->CapturedMouseMove(InViewportClient, InViewport, InMouseX, InMouseY);
+	return bHandled;
+}
+
+
+bool FModelingToolsEditorMode::EndTracking(FEditorViewportClient* InViewportClient, FViewport* InViewport)
+{
+	bool bHandled = ToolsContext->EndTracking(InViewportClient, InViewport);
+	return bHandled;
+}
+
+
+
+
+
+
+
+
+bool FModelingToolsEditorMode::ReceivedFocus(FEditorViewportClient* ViewportClient, FViewport* Viewport)
+{
+#ifdef ENABLE_DEBUG_PRINTING
+	UE_LOG(LogTemp, Warning, TEXT("RECEIVED FOCUS"));
+#endif
+
+	return false;
+}
+
+bool FModelingToolsEditorMode::LostFocus(FEditorViewportClient* ViewportClient, FViewport* Viewport)
+{
+#ifdef ENABLE_DEBUG_PRINTING
+	UE_LOG(LogTemp, Warning, TEXT("LOST FOCUS"));
+#endif
+
+	return false;
+}
+
+
+
+bool FModelingToolsEditorMode::MouseEnter(FEditorViewportClient* ViewportClient, FViewport* Viewport, int32 x, int32 y)
+{
+	bool bHandled = ToolsContext->MouseEnter(ViewportClient, Viewport, x, y);
+	return bHandled;
+}
+
+bool FModelingToolsEditorMode::MouseMove(FEditorViewportClient* ViewportClient, FViewport* Viewport, int32 x, int32 y)
+{
+	bool bHandled = ToolsContext->MouseMove(ViewportClient, Viewport, x, y);
+	return bHandled;
+}
+
+bool FModelingToolsEditorMode::MouseLeave(FEditorViewportClient* ViewportClient, FViewport* Viewport)
+{
+	bool bHandled = ToolsContext->MouseLeave(ViewportClient, Viewport);
+	return bHandled;
+}
+
+
+
+void FModelingToolsEditorMode::Enter()
+{
+	FEdMode::Enter();
+
+	// initialize FEdMode ToolsContext adapter
+
+	//ToolsContext = NewObject<UEdModeInteractiveToolsContext>(GetTransientPackage(), TEXT("ToolsContext"), RF_Transient);
+	ToolsContext = NewObject<UEdModeInteractiveToolsContext>();
+	ToolsContext->InitializeContextFromEdMode(this);
+
+	ToolsContext->OnToolNotificationMessage.AddLambda([this](const FText& Message)
+	{
+		this->OnToolNotificationMessage.Broadcast(Message);
+	});
+	ToolsContext->OnToolWarningMessage.AddLambda([this](const FText& Message)
+	{
+		this->OnToolWarningMessage.Broadcast(Message);
+	});
+
+	if (!Toolkit.IsValid() && UsesToolkits())
+	{
+		Toolkit = MakeShareable(new FModelingToolsEditorModeToolkit);
+		Toolkit->Init(Owner->GetToolkitHost());
+
+
+		const FModelingToolsManagerCommands& ToolManagerCommands = FModelingToolsManagerCommands::Get();
+		const TSharedRef<FUICommandList>& CommandList = Toolkit->GetToolkitCommands();
+
+		CommandList->MapAction(
+			ToolManagerCommands.AcceptActiveTool,
+			FExecuteAction::CreateLambda([this]() { ToolsContext->EndTool(EToolShutdownType::Accept); }),
+			FCanExecuteAction::CreateLambda([this]() { return ToolsContext->CanAcceptActiveTool(); }),
+			FGetActionCheckState(),
+			FIsActionButtonVisible::CreateLambda([this]() {return ToolsContext->ActiveToolHasAccept(); }),
+			EUIActionRepeatMode::RepeatDisabled
+		);
+
+		CommandList->MapAction(
+			ToolManagerCommands.CancelActiveTool,
+			FExecuteAction::CreateLambda([this]() { ToolsContext->EndTool(EToolShutdownType::Cancel); }),
+			FCanExecuteAction::CreateLambda([this]() { return ToolsContext->CanCancelActiveTool(); }),
+			FGetActionCheckState(),
+			FIsActionButtonVisible::CreateLambda([this]() {return ToolsContext->ActiveToolHasAccept(); }),
+			EUIActionRepeatMode::RepeatDisabled
+		);
+
+		CommandList->MapAction(
+			ToolManagerCommands.CompleteActiveTool,
+			FExecuteAction::CreateLambda([this]() { ToolsContext->EndTool(EToolShutdownType::Completed); }),
+			FCanExecuteAction::CreateLambda([this]() { return ToolsContext->CanCompleteActiveTool(); }),
+			FGetActionCheckState(),
+			FIsActionButtonVisible::CreateLambda([this]() {return ToolsContext->CanCompleteActiveTool(); }),
+			EUIActionRepeatMode::RepeatDisabled
+		);
+
+	}
+
+	const FModelingToolsManagerCommands& ToolManagerCommands = FModelingToolsManagerCommands::Get();
+	const TSharedRef<FUICommandList>& CommandList = Toolkit->GetToolkitCommands();
+	auto RegisterToolFunc = [this, &CommandList](TSharedPtr<FUICommandInfo> UICommand, FString ToolIdentifier, UInteractiveToolBuilder* Builder)
+	{
+		ToolsContext->ToolManager->RegisterToolType(ToolIdentifier, Builder);
+		CommandList->MapAction( UICommand,
+			FExecuteAction::CreateLambda([this, ToolIdentifier]() { ToolsContext->StartTool(ToolIdentifier); }),
+			FCanExecuteAction::CreateLambda([this, ToolIdentifier]() { return ToolsContext->CanStartTool(ToolIdentifier); }));
+	};
+
+
+	// register tool set
+
+	//
+	// make shape tools
+	//
+	auto AddPrimitiveToolBuilder = NewObject<UAddPrimitiveToolBuilder>();
+	AddPrimitiveToolBuilder->AssetAPI = ToolsContext->GetAssetAPI();
+	RegisterToolFunc(ToolManagerCommands.BeginAddPrimitiveTool, TEXT("AddPrimitiveTool"), AddPrimitiveToolBuilder);
+
+	auto DrawPolygonToolBuilder = NewObject<UDrawPolygonToolBuilder>();
+	DrawPolygonToolBuilder->AssetAPI = ToolsContext->GetAssetAPI();
+	RegisterToolFunc(ToolManagerCommands.BeginDrawPolygonTool, TEXT("DrawPolygonTool"), DrawPolygonToolBuilder);
+
+	auto ShapeSprayToolBuilder = NewObject<UShapeSprayToolBuilder>();
+	ShapeSprayToolBuilder->AssetAPI = ToolsContext->GetAssetAPI();
+	RegisterToolFunc(ToolManagerCommands.BeginShapeSprayTool, TEXT("ShapeSprayTool"), ShapeSprayToolBuilder);
+
+
+	//
+	// vertex deform tools
+	// 
+
+	auto MoveVerticesToolBuilder = NewObject<UDynamicMeshSculptToolBuilder>();
+	MoveVerticesToolBuilder->bEnableRemeshing = false;
+	RegisterToolFunc(ToolManagerCommands.BeginSculptMeshTool, TEXT("MoveVerticesTool"), MoveVerticesToolBuilder);
+
+	RegisterToolFunc(ToolManagerCommands.BeginPolyEditTool, TEXT("EditMeshPolygonsTool"), NewObject<UEditMeshPolygonsToolBuilder>());
+	RegisterToolFunc(ToolManagerCommands.BeginSmoothMeshTool, TEXT("SmoothMeshTool"), NewObject<USmoothMeshToolBuilder>());
+	RegisterToolFunc(ToolManagerCommands.BeginDisplaceMeshTool, TEXT("DisplaceMeshTool"), NewObject<UDisplaceMeshToolBuilder>());
+	RegisterToolFunc(ToolManagerCommands.BeginMeshSpaceDeformerTool, TEXT("MeshSpaceDeformerTool"), NewObject<UMeshSpaceDeformerToolBuilder>());
+	RegisterToolFunc(ToolManagerCommands.BeginTransformMeshesTool, TEXT("TransformMeshesTool"), NewObject<UTransformMeshesToolBuilder>());
+
+	// edit tools
+
+
+	auto DynaSculptToolBuilder = NewObject<UDynamicMeshSculptToolBuilder>();
+	DynaSculptToolBuilder->bEnableRemeshing = true;
+	RegisterToolFunc(ToolManagerCommands.BeginRemeshSculptMeshTool, TEXT("DynaSculptTool"), DynaSculptToolBuilder);
+
+	RegisterToolFunc(ToolManagerCommands.BeginRemeshMeshTool, TEXT("RemeshMeshTool"), NewObject<URemeshMeshToolBuilder>());
+	RegisterToolFunc(ToolManagerCommands.BeginSimplifyMeshTool, TEXT("SimplifyMeshTool"), NewObject<USimplifyMeshToolBuilder>());
+
+
+	auto EditNormalsToolBuilder = NewObject<UEditNormalsToolBuilder>();
+	EditNormalsToolBuilder->AssetAPI = ToolsContext->GetAssetAPI();
+	RegisterToolFunc(ToolManagerCommands.BeginEditNormalsTool, TEXT("EditNormalsTool"), EditNormalsToolBuilder);
+
+	auto UVProjectionToolBuilder = NewObject<UUVProjectionToolBuilder>();
+	UVProjectionToolBuilder->AssetAPI = ToolsContext->GetAssetAPI();
+	RegisterToolFunc(ToolManagerCommands.BeginUVProjectionTool, TEXT("UVProjectionTool"), UVProjectionToolBuilder);
+
+	auto MergeMeshesToolBuilder = NewObject<UMergeMeshesToolBuilder>();
+	MergeMeshesToolBuilder->AssetAPI = ToolsContext->GetAssetAPI();
+	RegisterToolFunc(ToolManagerCommands.BeginVoxelMergeTool, TEXT("MergeMeshesTool"), MergeMeshesToolBuilder);
+
+	auto VoxelCSGMeshesToolBuilder = NewObject<UVoxelCSGMeshesToolBuilder>();
+	VoxelCSGMeshesToolBuilder->AssetAPI = ToolsContext->GetAssetAPI();
+	RegisterToolFunc(ToolManagerCommands.BeginVoxelBooleanTool, TEXT("VoxelCSGMeshesTool"), VoxelCSGMeshesToolBuilder);
+
+	auto PlaneCutToolBuilder = NewObject<UPlaneCutToolBuilder>();
+	PlaneCutToolBuilder->AssetAPI = ToolsContext->GetAssetAPI();
+	RegisterToolFunc(ToolManagerCommands.BeginPlaneCutTool, TEXT("PlaneCutTool"), PlaneCutToolBuilder);
+
+	auto PolygonOnMeshToolBuilder = NewObject<UPolygonOnMeshToolBuilder>();
+	PolygonOnMeshToolBuilder->AssetAPI = ToolsContext->GetAssetAPI();
+	RegisterToolFunc(ToolManagerCommands.BeginPolygonOnMeshTool, TEXT("PolygonOnMeshTool"), PolygonOnMeshToolBuilder);
+
+	auto ParameterizeMeshToolBuilder = NewObject<UParameterizeMeshToolBuilder>();
+	ParameterizeMeshToolBuilder->AssetAPI = ToolsContext->GetAssetAPI();
+	RegisterToolFunc(ToolManagerCommands.BeginParameterizeMeshTool, TEXT("ParameterizeMeshTool"), ParameterizeMeshToolBuilder);
+
+	auto MeshSelectionToolBuilder = NewObject<UMeshSelectionToolBuilder>();
+	MeshSelectionToolBuilder->AssetAPI = ToolsContext->GetAssetAPI();
+	RegisterToolFunc(ToolManagerCommands.BeginMeshSelectionTool, TEXT("MeshSelectionTool"), MeshSelectionToolBuilder);
+
+
+	// analysis tools
+
+	RegisterToolFunc(ToolManagerCommands.BeginMeshInspectorTool, TEXT("MeshInspectorTool"), NewObject<UMeshInspectorToolBuilder>());
+	RegisterToolFunc(ToolManagerCommands.BeginWeldEdgesTool, TEXT("WeldMeshEdgesTool"), NewObject<UWeldMeshEdgesToolBuilder>());
+	RegisterToolFunc(ToolManagerCommands.BeginPolyGroupsTool, TEXT("ConvertToPolygonsTool"), NewObject<UConvertToPolygonsToolBuilder>());
+	RegisterToolFunc(ToolManagerCommands.BeginAttributeEditorTool, TEXT("AttributeEditorTool"), NewObject<UAttributeEditorToolBuilder>());
+
+
+
+	ToolsContext->ToolManager->SelectActiveToolType(EToolSide::Left, TEXT("DynaSculptTool"));
+
+	// listen for Tool start/end events to bind/unbind any hotkeys relevant to that Tool
+	ToolsContext->ToolManager->OnToolStarted.AddLambda([this](UInteractiveToolManager* Manager, UInteractiveTool* Tool)
+	{
+		FModelingToolActionCommands::UpdateToolCommandBinding(Tool, UICommandList, false);
+	});
+	ToolsContext->ToolManager->OnToolEnded.AddLambda([this](UInteractiveToolManager* Manager, UInteractiveTool* Tool)
+	{
+		FModelingToolActionCommands::UpdateToolCommandBinding(Tool, UICommandList, true);
+	});
+
+
+#ifdef TOOLED_ENABLE_VIEWPORT_INTERACTION
+	///
+	// Viewport Interaction
+	///
+	UEditorWorldExtensionCollection* ExtensionCollection = GEditor->GetEditorWorldExtensionsManager()->GetEditorWorldExtensions(GetWorld());
+	check(ExtensionCollection != nullptr);
+	this->ViewportWorldInteraction = NewObject<UViewportWorldInteraction>(ExtensionCollection);
+	ExtensionCollection->AddExtension(this->ViewportWorldInteraction);
+		//Cast<UViewportWorldInteraction>(ExtensionCollection->AddExtension(UViewportWorldInteraction::StaticClass()));
+	check(ViewportWorldInteraction != nullptr);
+	//this->ViewportWorldInteraction->UseLegacyInteractions();
+	//this->ViewportWorldInteraction->AddMouseCursorInteractor();
+	this->ViewportWorldInteraction->SetUseInputPreprocessor(true);
+	this->ViewportWorldInteraction->SetGizmoHandleType(EGizmoHandleTypes::All);
+
+	// Set the current viewport.
+	{
+		const TSharedRef< ILevelEditor >& LevelEditor = FModuleManager::GetModuleChecked<FLevelEditorModule>("LevelEditor").GetFirstLevelEditor().ToSharedRef();
+
+		// Do we have an active perspective viewport that is valid for VR?  If so, go ahead and use that.
+		TSharedPtr<FEditorViewportClient> ViewportClient;
+		{
+			TSharedPtr<IAssetViewport> ActiveLevelViewport = LevelEditor->GetActiveViewportInterface();
+			if (ActiveLevelViewport.IsValid())
+			{
+				ViewportClient = StaticCastSharedRef<SLevelViewport>(ActiveLevelViewport->AsWidget())->GetViewportClient();
+			}
+		}
+
+		this->ViewportWorldInteraction->SetDefaultOptionalViewportClient(ViewportClient);
+	}
+#endif  // TOOLED_ENABLE_VIEWPORT_INTERACTION
+}
+
+
+
+void FModelingToolsEditorMode::Exit()
+{
+	OnToolNotificationMessage.Clear();
+	OnToolWarningMessage.Clear();
+
+	ToolsContext->ShutdownContext();
+	ToolsContext = nullptr;
+
+	if (Toolkit.IsValid())
+	{
+		FToolkitManager::Get().CloseToolkit(Toolkit.ToSharedRef());
+		Toolkit.Reset();
+	}
+
+
+#ifdef TOOLED_ENABLE_VIEWPORT_INTERACTION
+	///
+	// Viewport Interaction
+	//
+	if (IViewportInteractionModule::IsAvailable())
+	{
+		if (ViewportWorldInteraction != nullptr)
+		{
+			ViewportWorldInteraction->ReleaseMouseCursorInteractor();
+
+			// Make sure gizmo is visible.  We may have hidden it
+			ViewportWorldInteraction->SetTransformGizmoVisible(true);
+
+			// Unregister mesh element transformer
+			//ViewportWorldInteraction->SetTransformer(nullptr);
+
+			UEditorWorldExtensionCollection* ExtensionCollection = GEditor->GetEditorWorldExtensionsManager()->GetEditorWorldExtensions(GetWorld());
+			if (ExtensionCollection != nullptr)
+			{
+				ExtensionCollection->RemoveExtension(ViewportWorldInteraction);
+			}
+
+			ViewportWorldInteraction = nullptr;
+		}
+	}
+#endif // TOOLED_ENABLE_VIEWPORT_INTERACTION
+
+
+	// Call base Exit method to ensure proper cleanup
+	FEdMode::Exit();
+}
+
+bool FModelingToolsEditorMode::UsesToolkits() const
+{
+	return true;
+}
+
+
+void FModelingToolsEditorMode::AddReferencedObjects(FReferenceCollector& Collector)
+{
+	Collector.AddReferencedObject(ToolsContext);
+}
+
+
+
+
+
+#undef LOCTEXT_NAMESPACE
\ No newline at end of file
diff --git a/Engine/Plugins/Experimental/ModelingToolsEditorMode/Source/ModelingToolsEditorMode/Private/ModelingToolsEditorModeModule.cpp b/Engine/Plugins/Experimental/ModelingToolsEditorMode/Source/ModelingToolsEditorMode/Private/ModelingToolsEditorModeModule.cpp
new file mode 100644
--- /dev/null
+++ b/Engine/Plugins/Experimental/ModelingToolsEditorMode/Source/ModelingToolsEditorMode/Private/ModelingToolsEditorModeModule.cpp
@@ -0,0 +1,51 @@
+// Copyright 1998-2019 Epic Games, Inc. All Rights Reserved.
+
+#include "ModelingToolsEditorModeModule.h"
+#include "ModelingToolsEditorMode.h"
+
+#include "ModelingToolsActions.h"
+#include "ModelingToolsManagerActions.h"
+#include "ModelingToolsEditorModeStyle.h"
+
+#define LOCTEXT_NAMESPACE "FModelingToolsEditorModeModule"
+
+void FModelingToolsEditorModeModule::StartupModule()
+{
+	FCoreDelegates::OnPostEngineInit.AddRaw(this, &FModelingToolsEditorModeModule::OnPostEngineInit);
+}
+
+void FModelingToolsEditorModeModule::ShutdownModule()
+{
+	FCoreDelegates::OnPostEngineInit.RemoveAll(this);
+
+	FModelingToolActionCommands::UnregisterAllToolActions();
+	FModelingToolsManagerCommands::Unregister();
+
+	// Unregister slate style overrides
+	FModelingToolsEditorModeStyle::Shutdown();
+
+	// This function may be called during shutdown to clean up your module.  For modules that support dynamic reloading,
+	// we call this function before unloading the module.
+	FEditorModeRegistry::Get().UnregisterMode(FModelingToolsEditorMode::EM_ModelingToolsEditorModeId);
+}
+
+
+void FModelingToolsEditorModeModule::OnPostEngineInit()
+{
+	// Register slate style overrides
+	FModelingToolsEditorModeStyle::Initialize();
+
+	// This code will execute after your module is loaded into memory; the exact timing is specified in the .uplugin file per-module
+	FEditorModeRegistry::Get().RegisterMode<FModelingToolsEditorMode>(
+		FModelingToolsEditorMode::EM_ModelingToolsEditorModeId,
+		LOCTEXT("ModelingToolsEditorModeName", "ModelingToolsEditorMode"),
+		FSlateIcon(),
+		true);
+
+	FModelingToolActionCommands::RegisterAllToolActions();
+	FModelingToolsManagerCommands::Register();
+}
+
+#undef LOCTEXT_NAMESPACE
+	
+IMPLEMENT_MODULE(FModelingToolsEditorModeModule, ModelingToolsEditorMode)
\ No newline at end of file
