Merge CL#9408555

Geometry Modeling Toolbar Icons

diff --git a/Engine/Plugins/Experimental/ModelingToolsEditorMode/Source/ModelingToolsEditorMode/Private/ModelingToolsEditorModeStyle.cpp b/Engine/Plugins/Experimental/ModelingToolsEditorMode/Source/ModelingToolsEditorMode/Private/ModelingToolsEditorModeStyle.cpp
--- a/Engine/Plugins/Experimental/ModelingToolsEditorMode/Source/ModelingToolsEditorMode/Private/ModelingToolsEditorModeStyle.cpp
+++ b/Engine/Plugins/Experimental/ModelingToolsEditorMode/Source/ModelingToolsEditorMode/Private/ModelingToolsEditorModeStyle.cpp
@@ -69,106 +69,71 @@ void FModelingToolsEditorModeStyle::Initialize()
 	// Tool Manager icons
 	{
 		// Accept/Cancel/Complete active tool
-		StyleSet->Set("ModelingToolsManagerCommands.CancelActiveTool", new IMAGE_PLUGIN_BRUSH("Icons/icon_ActiveTool_Cancel_40x", Icon28x28));
-		StyleSet->Set("ModelingToolsManagerCommands.CancelActiveTool.Small", new IMAGE_PLUGIN_BRUSH("Icons/icon_ActiveTool_Cancel_40x", Icon28x28));
-		StyleSet->Set("ModelingToolsManagerCommands.AcceptActiveTool", new IMAGE_PLUGIN_BRUSH("Icons/icon_ActiveTool_Accept_40x", Icon28x28));
-		StyleSet->Set("ModelingToolsManagerCommands.AcceptActiveTool.Small", new IMAGE_PLUGIN_BRUSH("Icons/icon_ActiveTool_Accept_40x", Icon28x28));
-		StyleSet->Set("ModelingToolsManagerCommands.CompleteActiveTool", new IMAGE_PLUGIN_BRUSH("Icons/icon_ActiveTool_Accept_40x", Icon28x28));
-		StyleSet->Set("ModelingToolsManagerCommands.CompleteActiveTool.Small", new IMAGE_PLUGIN_BRUSH("Icons/icon_ActiveTool_Accept_40x", Icon28x28));
-
-
-		StyleSet->Set("ModelingToolsManagerCommands.BeginAddPrimitiveTool", new IMAGE_PLUGIN_BRUSH("Icons/icon_Tool_MakePrimitive_40x", Icon28x28));
-		StyleSet->Set("ModelingToolsManagerCommands.BeginAddPrimitiveTool.Small", new IMAGE_PLUGIN_BRUSH("Icons/icon_Tool_MakePrimitive_40x", Icon28x28));
-		StyleSet->Set("ModelingToolsManagerCommands.BeginDrawPolygonTool", new IMAGE_PLUGIN_BRUSH("Icons/icon_Tool_DrawPolygon_40x", Icon28x28));
-		StyleSet->Set("ModelingToolsManagerCommands.BeginDrawPolygonTool.Small", new IMAGE_PLUGIN_BRUSH("Icons/icon_Tool_DrawPolygon_40x", Icon28x28));
-		StyleSet->Set("ModelingToolsManagerCommands.BeginShapeSprayTool", new IMAGE_PLUGIN_BRUSH("Icons/icon_Tool_ShapeSpray_40x", Icon28x28));
-		StyleSet->Set("ModelingToolsManagerCommands.BeginShapeSprayTool.Small", new IMAGE_PLUGIN_BRUSH("Icons/icon_Tool_ShapeSpray_40x", Icon28x28));
-
-		StyleSet->Set("ModelingToolsManagerCommands.BeginSmoothMeshTool", new IMAGE_PLUGIN_BRUSH("Icons/icon_Tool_Smooth_40x", Icon28x28));
-		StyleSet->Set("ModelingToolsManagerCommands.BeginSmoothMeshTool.Small", new IMAGE_PLUGIN_BRUSH("Icons/icon_Tool_Smooth_40x", Icon28x28));
-		StyleSet->Set("ModelingToolsManagerCommands.BeginSculptMeshTool", new IMAGE_PLUGIN_BRUSH("Icons/icon_Tool_Sculpt_40x", Icon28x28));
-		StyleSet->Set("ModelingToolsManagerCommands.BeginSculptMeshTool.Small", new IMAGE_PLUGIN_BRUSH("Icons/icon_Tool_Sculpt_40x", Icon28x28));
-		StyleSet->Set("ModelingToolsManagerCommands.BeginPolyEditTool", new IMAGE_PLUGIN_BRUSH("Icons/icon_Tool_PolyEdit_40x", Icon28x28));
-		StyleSet->Set("ModelingToolsManagerCommands.BeginPolyEditTool.Small", new IMAGE_PLUGIN_BRUSH("Icons/icon_Tool_PolyEdit_40x", Icon28x28));
-		StyleSet->Set("ModelingToolsManagerCommands.BeginDisplaceMeshTool", new IMAGE_PLUGIN_BRUSH("Icons/icon_Tool_Displace_40x", Icon28x28));
-		StyleSet->Set("ModelingToolsManagerCommands.BeginDisplaceMeshTool.Small", new IMAGE_PLUGIN_BRUSH("Icons/icon_Tool_Displace_40x", Icon28x28));
-		StyleSet->Set("ModelingToolsManagerCommands.BeginMeshSpaceDeformerTool", new IMAGE_PLUGIN_BRUSH("Icons/icon_Tool_Displace_40x", Icon28x28));
-		StyleSet->Set("ModelingToolsManagerCommands.BeginMeshSpaceDeformerTool.Small", new IMAGE_PLUGIN_BRUSH("Icons/icon_Tool_Displace_40x", Icon28x28));
-		StyleSet->Set("ModelingToolsManagerCommands.BeginTransformMeshesTool", new IMAGE_PLUGIN_BRUSH("Icons/icon_Tool_Transform_40x", Icon28x28));
-		StyleSet->Set("ModelingToolsManagerCommands.BeginTransformMeshesTool.Small", new IMAGE_PLUGIN_BRUSH("Icons/icon_Tool_Transform_40x", Icon28x28));
-
-		StyleSet->Set("ModelingToolsManagerCommands.BeginRemeshSculptMeshTool", new IMAGE_PLUGIN_BRUSH("Icons/icon_Tool_RemeshSculpt_40x", Icon28x28));
-		StyleSet->Set("ModelingToolsManagerCommands.BeginRemeshSculptMeshTool.Small", new IMAGE_PLUGIN_BRUSH("Icons/icon_Tool_RemeshSculpt_40x", Icon28x28));
-		StyleSet->Set("ModelingToolsManagerCommands.BeginRemeshMeshTool", new IMAGE_PLUGIN_BRUSH("Icons/icon_Tool_Remesh_40x", Icon28x28));
-		StyleSet->Set("ModelingToolsManagerCommands.BeginRemeshMeshTool.Small", new IMAGE_PLUGIN_BRUSH("Icons/icon_Tool_Remesh_40x", Icon28x28));
-		StyleSet->Set("ModelingToolsManagerCommands.BeginSimplifyMeshTool", new IMAGE_PLUGIN_BRUSH("Icons/icon_Tool_Simplify_40x", Icon28x28));
-		StyleSet->Set("ModelingToolsManagerCommands.BeginSimplifyMeshTool.Small", new IMAGE_PLUGIN_BRUSH("Icons/icon_Tool_Simplify_40x", Icon28x28));
-		StyleSet->Set("ModelingToolsManagerCommands.BeginEditNormalsTool", new IMAGE_PLUGIN_BRUSH("Icons/icon_Tool_EditNormals_40x", Icon28x28));
-		StyleSet->Set("ModelingToolsManagerCommands.BeginEditNormalsTool.Small", new IMAGE_PLUGIN_BRUSH("Icons/icon_Tool_EditNormals_40x", Icon28x28));
-		StyleSet->Set("ModelingToolsManagerCommands.BeginUVProjectionTool", new IMAGE_PLUGIN_BRUSH("Icons/icon_Tool_UVProjection_40x", Icon28x28));
-		StyleSet->Set("ModelingToolsManagerCommands.BeginUVProjectionTool.Small", new IMAGE_PLUGIN_BRUSH("Icons/icon_Tool_UVProjection_40x", Icon28x28));
-		StyleSet->Set("ModelingToolsManagerCommands.BeginVoxelMergeTool", new IMAGE_PLUGIN_BRUSH("Icons/icon_Tool_VoxMerge_40x", Icon28x28));
-		StyleSet->Set("ModelingToolsManagerCommands.BeginVoxelMergeTool.Small", new IMAGE_PLUGIN_BRUSH("Icons/icon_Tool_VoxMerge_40x", Icon28x28));
-		StyleSet->Set("ModelingToolsManagerCommands.BeginVoxelBooleanTool", new IMAGE_PLUGIN_BRUSH("Icons/icon_Tool_VoxBoolean_40x", Icon28x28));
-		StyleSet->Set("ModelingToolsManagerCommands.BeginVoxelBooleanTool.Small", new IMAGE_PLUGIN_BRUSH("Icons/icon_Tool_VoxBoolean_40x", Icon28x28));
-		StyleSet->Set("ModelingToolsManagerCommands.BeginPlaneCutTool", new IMAGE_PLUGIN_BRUSH("Icons/icon_Tool_PlaneCut_40x", Icon28x28));
-		StyleSet->Set("ModelingToolsManagerCommands.BeginPlaneCutTool.Small", new IMAGE_PLUGIN_BRUSH("Icons/icon_Tool_PlaneCut_40x", Icon28x28));
-		StyleSet->Set("ModelingToolsManagerCommands.BeginPolygonOnMeshTool", new IMAGE_PLUGIN_BRUSH("Icons/icon_Tool_PolygonOnMesh_40x", Icon28x28));
-		StyleSet->Set("ModelingToolsManagerCommands.BeginPolygonOnMeshTool.Small", new IMAGE_PLUGIN_BRUSH("Icons/icon_Tool_PolygonOnMesh_40x", Icon28x28));
-		StyleSet->Set("ModelingToolsManagerCommands.BeginMeshSelectionTool", new IMAGE_PLUGIN_BRUSH("Icons/icon_Tool_MeshSelection_40x", Icon28x28));
-		StyleSet->Set("ModelingToolsManagerCommands.BeginMeshSelectionTool.Small", new IMAGE_PLUGIN_BRUSH("Icons/icon_Tool_MeshSelection_40x", Icon28x28));
-
-		StyleSet->Set("ModelingToolsManagerCommands.BeginMeshInspectorTool", new IMAGE_PLUGIN_BRUSH("Icons/icon_Tool_MeshInspector_40x", Icon28x28));
-		StyleSet->Set("ModelingToolsManagerCommands.BeginMeshInspectorTool.Small", new IMAGE_PLUGIN_BRUSH("Icons/icon_Tool_MeshInspector_40x", Icon28x28));
-		StyleSet->Set("ModelingToolsManagerCommands.BeginParameterizeMeshTool", new IMAGE_PLUGIN_BRUSH("Icons/icon_Tool_UVGenerate_40x", Icon28x28));
-		StyleSet->Set("ModelingToolsManagerCommands.BeginParameterizeMeshTool.Small", new IMAGE_PLUGIN_BRUSH("Icons/icon_Tool_UVGenerate_40x", Icon28x28));
-		StyleSet->Set("ModelingToolsManagerCommands.BeginWeldEdgesTool", new IMAGE_PLUGIN_BRUSH("Icons/icon_Tool_WeldEdges_40x", Icon28x28));
-		StyleSet->Set("ModelingToolsManagerCommands.BeginWeldEdgesTool.Small", new IMAGE_PLUGIN_BRUSH("Icons/icon_Tool_WeldEdges_40x", Icon28x28));
-		StyleSet->Set("ModelingToolsManagerCommands.BeginPolyGroupsTool", new IMAGE_PLUGIN_BRUSH("Icons/icon_Tool_PolyGroups_40x", Icon28x28));
-		StyleSet->Set("ModelingToolsManagerCommands.BeginPolyGroupsTool.Small", new IMAGE_PLUGIN_BRUSH("Icons/icon_Tool_PolyGroups_40x", Icon28x28));
-		StyleSet->Set("ModelingToolsManagerCommands.BeginAttributeEditorTool", new IMAGE_PLUGIN_BRUSH("Icons/icon_Tool_AttributeEditor_40x", Icon28x28));
-		StyleSet->Set("ModelingToolsManagerCommands.BeginAttributeEditorTool.Small", new IMAGE_PLUGIN_BRUSH("Icons/icon_Tool_AttributeEditor_40x", Icon28x28));
-
-
-
-		//StyleSet->Set("ModelingToolsManagerCommands.CancelActiveTool", new IMAGE_PLUGIN_BRUSH("Icons/icon_ActiveTool_Cancel_40x", Icon40x40));
-		//StyleSet->Set("ModelingToolsManagerCommands.CancelActiveTool.Small", new IMAGE_PLUGIN_BRUSH("Icons/icon_ActiveTool_Cancel_20x", Icon20x20));
-		//StyleSet->Set("ModelingToolsManagerCommands.AcceptActiveTool", new IMAGE_PLUGIN_BRUSH("Icons/icon_ActiveTool_Accept_40x", Icon40x40));
-		//StyleSet->Set("ModelingToolsManagerCommands.AcceptActiveTool.Small", new IMAGE_PLUGIN_BRUSH("Icons/icon_ActiveTool_Accept_20x", Icon20x20));
-		//StyleSet->Set("ModelingToolsManagerCommands.CompleteActiveTool", new IMAGE_PLUGIN_BRUSH("Icons/icon_ActiveTool_Accept_40x", Icon40x40));
-		//StyleSet->Set("ModelingToolsManagerCommands.CompleteActiveTool.Small", new IMAGE_PLUGIN_BRUSH("Icons/icon_ActiveTool_Accept_20x", Icon20x20));
-
-
-		//StyleSet->Set("ModelingToolsManagerCommands.BeginMakePrimitiveTool", new IMAGE_PLUGIN_BRUSH("Icons/icon_Tool_MakePrimitive_40x", Icon40x40));
-		//StyleSet->Set("ModelingToolsManagerCommands.BeginMakePrimitiveTool.Small", new IMAGE_PLUGIN_BRUSH("Icons/icon_Tool_MakePrimitive_20x", Icon20x20));
-		//StyleSet->Set("ModelingToolsManagerCommands.BeginDrawPolygonTool", new IMAGE_PLUGIN_BRUSH("Icons/icon_Tool_DrawPolygon_40x", Icon40x40));
-		//StyleSet->Set("ModelingToolsManagerCommands.BeginDrawPolygonTool.Small", new IMAGE_PLUGIN_BRUSH("Icons/icon_Tool_DrawPolygon_20x", Icon20x20));
-		//StyleSet->Set("ModelingToolsManagerCommands.BeginShapeSprayTool", new IMAGE_PLUGIN_BRUSH("Icons/icon_Tool_ShapeSpray_40x", Icon40x40));
-		//StyleSet->Set("ModelingToolsManagerCommands.BeginShapeSprayTool.Small", new IMAGE_PLUGIN_BRUSH("Icons/icon_Tool_ShapeSpray_20x", Icon20x20));
-
-		//StyleSet->Set("ModelingToolsManagerCommands.BeginSmoothMeshTool", new IMAGE_PLUGIN_BRUSH("Icons/icon_Tool_Smooth_40x", Icon40x40));
-		//StyleSet->Set("ModelingToolsManagerCommands.BeginSmoothMeshTool.Small", new IMAGE_PLUGIN_BRUSH("Icons/icon_Tool_Smooth_20x", Icon20x20));
-		//StyleSet->Set("ModelingToolsManagerCommands.BeginSculptMeshTool", new IMAGE_PLUGIN_BRUSH("Icons/icon_Tool_Sculpt_40x", Icon40x40));
-		//StyleSet->Set("ModelingToolsManagerCommands.BeginSculptMeshTool.Small", new IMAGE_PLUGIN_BRUSH("Icons/icon_Tool_Sculpt_20x", Icon20x20));
-		//StyleSet->Set("ModelingToolsManagerCommands.BeginPolyEditTool", new IMAGE_PLUGIN_BRUSH("Icons/icon_Tool_PolyEdit_40x", Icon40x40));
-		//StyleSet->Set("ModelingToolsManagerCommands.BeginPolyEditTool.Small", new IMAGE_PLUGIN_BRUSH("Icons/icon_Tool_PolyEdit_20x", Icon20x20));
-
-		//StyleSet->Set("ModelingToolsManagerCommands.BeginRemeshSculptMeshTool", new IMAGE_PLUGIN_BRUSH("Icons/icon_Tool_RemeshSculpt_40x", Icon40x40));
-		//StyleSet->Set("ModelingToolsManagerCommands.BeginRemeshSculptMeshTool.Small", new IMAGE_PLUGIN_BRUSH("Icons/icon_Tool_RemeshSculpt_20x", Icon20x20));
-		//StyleSet->Set("ModelingToolsManagerCommands.BeginRemeshMeshTool", new IMAGE_PLUGIN_BRUSH("Icons/icon_Tool_Remesh_40x", Icon40x40));
-		//StyleSet->Set("ModelingToolsManagerCommands.BeginRemeshMeshTool.Small", new IMAGE_PLUGIN_BRUSH("Icons/icon_Tool_Remesh_20x", Icon20x20));
-		//StyleSet->Set("ModelingToolsManagerCommands.BeginSimplifyMeshTool", new IMAGE_PLUGIN_BRUSH("Icons/icon_Tool_Simplify_40x", Icon40x40));
-		//StyleSet->Set("ModelingToolsManagerCommands.BeginSimplifyMeshTool.Small", new IMAGE_PLUGIN_BRUSH("Icons/icon_Tool_Simplify_20x", Icon20x20));
-
-
-		//StyleSet->Set("ModelingToolsManagerCommands.BeginMeshInspectorTool", new IMAGE_PLUGIN_BRUSH("Icons/icon_Tool_MeshInspector_40x", Icon40x40));
-		//StyleSet->Set("ModelingToolsManagerCommands.BeginMeshInspectorTool.Small", new IMAGE_PLUGIN_BRUSH("Icons/icon_Tool_MeshInspector_20x", Icon20x20));
-		//StyleSet->Set("ModelingToolsManagerCommands.BeginWeldEdgesTool", new IMAGE_PLUGIN_BRUSH("Icons/icon_Tool_WeldEdges_40x", Icon40x40));
-		//StyleSet->Set("ModelingToolsManagerCommands.BeginWeldEdgesTool.Small", new IMAGE_PLUGIN_BRUSH("Icons/icon_Tool_WeldEdges_20x", Icon20x20));
-		//StyleSet->Set("ModelingToolsManagerCommands.BeginPolyGroupsTool", new IMAGE_PLUGIN_BRUSH("Icons/icon_Tool_PolyGroups_40x", Icon40x40));
-		//StyleSet->Set("ModelingToolsManagerCommands.BeginPolyGroupsTool.Small", new IMAGE_PLUGIN_BRUSH("Icons/icon_Tool_PolyGroups_20x", Icon20x20));
 
 
+		// NOTE:  Old-style, need to be replaced: 
+		StyleSet->Set("ModelingToolsManagerCommands.CancelActiveTool", new IMAGE_PLUGIN_BRUSH("Icons/icon_ActiveTool_Cancel_40x", Icon20x20));
+		StyleSet->Set("ModelingToolsManagerCommands.CancelActiveTool.Small", new IMAGE_PLUGIN_BRUSH("Icons/icon_ActiveTool_Cancel_40x", Icon20x20));
+		StyleSet->Set("ModelingToolsManagerCommands.AcceptActiveTool", new IMAGE_PLUGIN_BRUSH("Icons/icon_ActiveTool_Accept_40x", Icon20x20));
+		StyleSet->Set("ModelingToolsManagerCommands.AcceptActiveTool.Small", new IMAGE_PLUGIN_BRUSH("Icons/icon_ActiveTool_Accept_40x", Icon20x20));
+		StyleSet->Set("ModelingToolsManagerCommands.CompleteActiveTool", new IMAGE_PLUGIN_BRUSH("Icons/icon_ActiveTool_Accept_40x", Icon20x20));
+		StyleSet->Set("ModelingToolsManagerCommands.CompleteActiveTool.Small", new IMAGE_PLUGIN_BRUSH("Icons/icon_ActiveTool_Accept_40x", Icon20x20));
+
+
+		StyleSet->Set("ModelingToolsManagerCommands.BeginShapeSprayTool", 				new IMAGE_PLUGIN_BRUSH("Icons/icon_Tool_ShapeSpray_40x",	Icon20x20));
+		StyleSet->Set("ModelingToolsManagerCommands.BeginShapeSprayTool.Small", 		new IMAGE_PLUGIN_BRUSH("Icons/icon_Tool_ShapeSpray_40x",	Icon20x20));
+		StyleSet->Set("ModelingToolsManagerCommands.BeginMeshSpaceDeformerTool", 		new IMAGE_PLUGIN_BRUSH("Icons/icon_Tool_Displace_40x",		Icon20x20));
+		StyleSet->Set("ModelingToolsManagerCommands.BeginMeshSpaceDeformerTool.Small", 	new IMAGE_PLUGIN_BRUSH("Icons/icon_Tool_Displace_40x",		Icon20x20));
+		StyleSet->Set("ModelingToolsManagerCommands.BeginPolygonOnMeshTool", 			new IMAGE_PLUGIN_BRUSH("Icons/icon_Tool_PolygonOnMesh_40x",	Icon20x20));
+		StyleSet->Set("ModelingToolsManagerCommands.BeginPolygonOnMeshTool.Small", 		new IMAGE_PLUGIN_BRUSH("Icons/icon_Tool_PolygonOnMesh_40x",	Icon20x20));
+		StyleSet->Set("ModelingToolsManagerCommands.BeginParameterizeMeshTool", 		new IMAGE_PLUGIN_BRUSH("Icons/icon_Tool_UVGenerate_40x",	Icon20x20));
+		StyleSet->Set("ModelingToolsManagerCommands.BeginParameterizeMeshTool.Small", 	new IMAGE_PLUGIN_BRUSH("Icons/icon_Tool_UVGenerate_40x",	Icon20x20));
+		StyleSet->Set("ModelingToolsManagerCommands.BeginPolyGroupsTool", 				new IMAGE_PLUGIN_BRUSH("Icons/icon_Tool_PolyGroups_40x",	Icon20x20));
+		StyleSet->Set("ModelingToolsManagerCommands.BeginPolyGroupsTool.Small", 		new IMAGE_PLUGIN_BRUSH("Icons/icon_Tool_PolyGroups_40x",	Icon20x20));
+
+
+		// Modes Palette Toolbar Icons
+		StyleSet->Set("ModelingToolsManagerCommands.BeginAddPrimitiveTool", 			new IMAGE_PLUGIN_BRUSH("Icons/Primitive_40x", 		Icon20x20));
+		StyleSet->Set("ModelingToolsManagerCommands.BeginAddPrimitiveTool.Small", 		new IMAGE_PLUGIN_BRUSH("Icons/Primitive_40x",		Icon20x20));
+		StyleSet->Set("ModelingToolsManagerCommands.BeginDrawPolygonTool", 				new IMAGE_PLUGIN_BRUSH("Icons/DrawPolygon_40x",		Icon20x20));
+		StyleSet->Set("ModelingToolsManagerCommands.BeginDrawPolygonTool.Small", 		new IMAGE_PLUGIN_BRUSH("Icons/DrawPolygon_40x", 	Icon20x20));
+
+		StyleSet->Set("ModelingToolsManagerCommands.BeginSmoothMeshTool", 				new IMAGE_PLUGIN_BRUSH("Icons/Smooth_40x", 			Icon20x20));
+		StyleSet->Set("ModelingToolsManagerCommands.BeginSmoothMeshTool.Small", 		new IMAGE_PLUGIN_BRUSH("Icons/Smooth_40x", 			Icon20x20));
+		StyleSet->Set("ModelingToolsManagerCommands.BeginSculptMeshTool", 				new IMAGE_PLUGIN_BRUSH("Icons/Sculpt_40x", 			Icon20x20));
+		StyleSet->Set("ModelingToolsManagerCommands.BeginSculptMeshTool.Small", 		new IMAGE_PLUGIN_BRUSH("Icons/Sculpt_40x", 			Icon20x20));
+		StyleSet->Set("ModelingToolsManagerCommands.BeginPolyEditTool", 				new IMAGE_PLUGIN_BRUSH("Icons/PolyEdit_40x", 		Icon20x20));
+		StyleSet->Set("ModelingToolsManagerCommands.BeginPolyEditTool.Small", 			new IMAGE_PLUGIN_BRUSH("Icons/PolyEdit_40x", 		Icon20x20));
+		StyleSet->Set("ModelingToolsManagerCommands.BeginDisplaceMeshTool", 			new IMAGE_PLUGIN_BRUSH("Icons/Displace_40x", 		Icon20x20));
+		StyleSet->Set("ModelingToolsManagerCommands.BeginDisplaceMeshTool.Small", 		new IMAGE_PLUGIN_BRUSH("Icons/Displace_40x", 		Icon20x20));
+		StyleSet->Set("ModelingToolsManagerCommands.BeginTransformMeshesTool", 			new IMAGE_PLUGIN_BRUSH("Icons/Transform_40x", 		Icon20x20));
+		StyleSet->Set("ModelingToolsManagerCommands.BeginTransformMeshesTool.Small", 	new IMAGE_PLUGIN_BRUSH("Icons/Transform_40x", 		Icon20x20));
+
+		StyleSet->Set("ModelingToolsManagerCommands.BeginRemeshSculptMeshTool", 		new IMAGE_PLUGIN_BRUSH("Icons/DynaSculpt_40x", 		Icon20x20));
+		StyleSet->Set("ModelingToolsManagerCommands.BeginRemeshSculptMeshTool.Small", 	new IMAGE_PLUGIN_BRUSH("Icons/DynaSculpt_40x",		Icon20x20));
+		StyleSet->Set("ModelingToolsManagerCommands.BeginRemeshMeshTool", 				new IMAGE_PLUGIN_BRUSH("Icons/Remesh_40x", 			Icon20x20));
+		StyleSet->Set("ModelingToolsManagerCommands.BeginRemeshMeshTool.Small", 		new IMAGE_PLUGIN_BRUSH("Icons/Remesh_40x", 			Icon20x20));
+		StyleSet->Set("ModelingToolsManagerCommands.BeginSimplifyMeshTool", 			new IMAGE_PLUGIN_BRUSH("Icons/Simplify_40x", 		Icon20x20));
+		StyleSet->Set("ModelingToolsManagerCommands.BeginSimplifyMeshTool.Small", 		new IMAGE_PLUGIN_BRUSH("Icons/Simplify_40x", 		Icon20x20));
+		StyleSet->Set("ModelingToolsManagerCommands.BeginEditNormalsTool", 				new IMAGE_PLUGIN_BRUSH("Icons/Normals_40x",			Icon20x20));
+		StyleSet->Set("ModelingToolsManagerCommands.BeginEditNormalsTool.Small", 		new IMAGE_PLUGIN_BRUSH("Icons/Normals_40x",			Icon20x20));
+		StyleSet->Set("ModelingToolsManagerCommands.BeginUVProjectionTool", 			new IMAGE_PLUGIN_BRUSH("Icons/UVProjection_40x", 	Icon20x20));
+		StyleSet->Set("ModelingToolsManagerCommands.BeginUVProjectionTool.Small", 		new IMAGE_PLUGIN_BRUSH("Icons/UVProjection_40x", 	Icon20x20));
+		StyleSet->Set("ModelingToolsManagerCommands.BeginVoxelMergeTool", 				new IMAGE_PLUGIN_BRUSH("Icons/VoxMerge_40x", 		Icon20x20));
+		StyleSet->Set("ModelingToolsManagerCommands.BeginVoxelMergeTool.Small", 		new IMAGE_PLUGIN_BRUSH("Icons/VoxMerge_40x", 		Icon20x20));
+		StyleSet->Set("ModelingToolsManagerCommands.BeginVoxelBooleanTool", 			new IMAGE_PLUGIN_BRUSH("Icons/VoxBoolean_40x", 		Icon20x20));
+		StyleSet->Set("ModelingToolsManagerCommands.BeginVoxelBooleanTool.Small", 		new IMAGE_PLUGIN_BRUSH("Icons/VoxBoolean_40x", 		Icon20x20));
+		StyleSet->Set("ModelingToolsManagerCommands.BeginPlaneCutTool", 				new IMAGE_PLUGIN_BRUSH("Icons/PlaneCut_40x", 		Icon20x20));
+		StyleSet->Set("ModelingToolsManagerCommands.BeginPlaneCutTool.Small", 			new IMAGE_PLUGIN_BRUSH("Icons/PlaneCut_40x", 		Icon20x20));
+		StyleSet->Set("ModelingToolsManagerCommands.BeginMeshSelectionTool", 			new IMAGE_PLUGIN_BRUSH("Icons/MeshSelect_40x",		Icon20x20));
+		StyleSet->Set("ModelingToolsManagerCommands.BeginMeshSelectionTool.Small", 		new IMAGE_PLUGIN_BRUSH("Icons/MeshSelect_40x",		Icon20x20));
+
+		StyleSet->Set("ModelingToolsManagerCommands.BeginMeshInspectorTool", 			new IMAGE_PLUGIN_BRUSH("Icons/Inspector_40x",		Icon20x20));
+		StyleSet->Set("ModelingToolsManagerCommands.BeginMeshInspectorTool.Small", 		new IMAGE_PLUGIN_BRUSH("Icons/Inspector_40x",		Icon20x20));
+		StyleSet->Set("ModelingToolsManagerCommands.BeginWeldEdgesTool", 				new IMAGE_PLUGIN_BRUSH("Icons/WeldEdges_40x", 		Icon20x20));
+		StyleSet->Set("ModelingToolsManagerCommands.BeginWeldEdgesTool.Small", 			new IMAGE_PLUGIN_BRUSH("Icons/WeldEdges_40x", 		Icon20x20));
+		StyleSet->Set("ModelingToolsManagerCommands.BeginAttributeEditorTool", 			new IMAGE_PLUGIN_BRUSH("Icons/AttributeEditor_40x", Icon20x20));
+		StyleSet->Set("ModelingToolsManagerCommands.BeginAttributeEditorTool.Small", 	new IMAGE_PLUGIN_BRUSH("Icons/AttributeEditor_40x", Icon20x20));
 
 
 		//const FLinearColor LayerSelectionColor = FLinearColor(0.13f, 0.70f, 1.00f);
