add new icon for modeling tools mode

diff --git a/Engine/Plugins/Experimental/ModelingToolsEditorMode/Source/ModelingToolsEditorMode/Private/ModelingToolsEditorModeModule.cpp b/Engine/Plugins/Experimental/ModelingToolsEditorMode/Source/ModelingToolsEditorMode/Private/ModelingToolsEditorModeModule.cpp
--- a/Engine/Plugins/Experimental/ModelingToolsEditorMode/Source/ModelingToolsEditorMode/Private/ModelingToolsEditorModeModule.cpp
+++ b/Engine/Plugins/Experimental/ModelingToolsEditorMode/Source/ModelingToolsEditorMode/Private/ModelingToolsEditorModeModule.cpp
@@ -39,7 +39,7 @@ void FModelingToolsEditorModeModule::OnPostEngineInit()
 	FEditorModeRegistry::Get().RegisterMode<FModelingToolsEditorMode>(
 		FModelingToolsEditorMode::EM_ModelingToolsEditorModeId,
 		LOCTEXT("ModelingToolsEditorModeName", "ModelingToolsEditorMode"),
-		FSlateIcon(),
+		FSlateIcon("ModelingToolsStyle", "LevelEditor.ModelingToolsMode", "LevelEditor.ModelingToolsMode.Small"),
 		true);
 
 	FModelingToolActionCommands::RegisterAllToolActions();
diff --git a/Engine/Plugins/Experimental/ModelingToolsEditorMode/Source/ModelingToolsEditorMode/Private/ModelingToolsEditorModeStyle.cpp b/Engine/Plugins/Experimental/ModelingToolsEditorMode/Source/ModelingToolsEditorMode/Private/ModelingToolsEditorModeStyle.cpp
--- a/Engine/Plugins/Experimental/ModelingToolsEditorMode/Source/ModelingToolsEditorMode/Private/ModelingToolsEditorModeStyle.cpp
+++ b/Engine/Plugins/Experimental/ModelingToolsEditorMode/Source/ModelingToolsEditorMode/Private/ModelingToolsEditorModeStyle.cpp
@@ -47,7 +47,7 @@ void FModelingToolsEditorModeStyle::Initialize()
 	}
 
 	StyleSet = MakeShareable(new FSlateStyleSet(GetStyleSetName()));
-	StyleSet->SetContentRoot(FPaths::EngineContentDir() / TEXT("Editor/Slate"));
+	StyleSet->SetContentRoot(FPaths::EnginePluginsDir() / TEXT("Experimental/ModelingToolsEditorMode/Content"));
 	StyleSet->SetCoreContentRoot(FPaths::EngineContentDir() / TEXT("Slate"));
 
 	const FTextBlockStyle& NormalText = FEditorStyle::Get().GetWidgetStyle<FTextBlockStyle>("NormalText");
@@ -70,6 +70,8 @@ void FModelingToolsEditorModeStyle::Initialize()
 	{
 		// Accept/Cancel/Complete active tool
 
+		StyleSet->Set("LevelEditor.ModelingToolsMode", new IMAGE_PLUGIN_BRUSH("Icons/icon_ModelingToolsEditorMode", FVector2D(40.0f, 40.0f)));
+		StyleSet->Set("LevelEditor.ModelingToolsMode.Small", new IMAGE_PLUGIN_BRUSH("Icons/icon_ModelingToolsEditorMode", FVector2D(40.0f, 40.0f)));
 
 		// NOTE:  Old-style, need to be replaced: 
 		StyleSet->Set("ModelingToolsManagerCommands.CancelActiveTool", new IMAGE_PLUGIN_BRUSH("Icons/icon_ActiveTool_Cancel_40x", Icon20x20));
