Disable HitProxy render pass inside modeling tools as no tools use it
#rnx

diff --git a/Engine/Plugins/Experimental/ModelingToolsEditorMode/Source/ModelingToolsEditorMode/Private/ModelingToolsEditorMode.cpp b/Engine/Plugins/Experimental/ModelingToolsEditorMode/Source/ModelingToolsEditorMode/Private/ModelingToolsEditorMode.cpp
--- a/Engine/Plugins/Experimental/ModelingToolsEditorMode/Source/ModelingToolsEditorMode/Private/ModelingToolsEditorMode.cpp
+++ b/Engine/Plugins/Experimental/ModelingToolsEditorMode/Source/ModelingToolsEditorMode/Private/ModelingToolsEditorMode.cpp
@@ -130,6 +130,12 @@ void FModelingToolsEditorMode::Render(const FSceneView* View, FViewport* Viewpor
 {
 	FEdMode::Render(View, Viewport, PDI);
 
+	// we do not use PDI hit testing in modeling tools, so skip these render passes
+	if (PDI->IsHitTesting())
+	{
+		return;
+	}
+
 	if (ToolsContext != nullptr)
 	{
 		ToolsContext->Render(View, Viewport, PDI);
