Change "ModelingToolsEditorMode" string to just "Modeling" in EdMode Tab Bar
#rnx

diff --git a/Engine/Plugins/Experimental/ModelingToolsEditorMode/Source/ModelingToolsEditorMode/Private/ModelingToolsEditorModeModule.cpp b/Engine/Plugins/Experimental/ModelingToolsEditorMode/Source/ModelingToolsEditorMode/Private/ModelingToolsEditorModeModule.cpp
--- a/Engine/Plugins/Experimental/ModelingToolsEditorMode/Source/ModelingToolsEditorMode/Private/ModelingToolsEditorModeModule.cpp
+++ b/Engine/Plugins/Experimental/ModelingToolsEditorMode/Source/ModelingToolsEditorMode/Private/ModelingToolsEditorModeModule.cpp
@@ -38,7 +38,7 @@ void FModelingToolsEditorModeModule::OnPostEngineInit()
 	// This code will execute after your module is loaded into memory; the exact timing is specified in the .uplugin file per-module
 	FEditorModeRegistry::Get().RegisterMode<FModelingToolsEditorMode>(
 		FModelingToolsEditorMode::EM_ModelingToolsEditorModeId,
-		LOCTEXT("ModelingToolsEditorModeName", "ModelingToolsEditorMode"),
+		LOCTEXT("ModelingToolsEditorModeName", "Modeling"),
 		FSlateIcon("ModelingToolsStyle", "LevelEditor.ModelingToolsMode", "LevelEditor.ModelingToolsMode.Small"),
 		true);
 
