Mesh Modeling Tools moved from NotForLicensees to Experimental

diff --git a/Engine/Plugins/Experimental/ModelingToolsEditorMode/ModelingToolsEditorMode.uplugin b/Engine/Plugins/Experimental/ModelingToolsEditorMode/ModelingToolsEditorMode.uplugin
new file mode 100644
--- /dev/null
+++ b/Engine/Plugins/Experimental/ModelingToolsEditorMode/ModelingToolsEditorMode.uplugin
@@ -0,0 +1,30 @@
+{
+	"FileVersion": 3,
+	"Version": 1,
+	"VersionName": "0.1",
+	"FriendlyName": "Modeling Tools Mode",
+	"Description": "Adds Modeling Tools Editor Mode Tab",
+	"Category": "Other",
+	"CreatedBy": "Epic Games, Inc.",
+	"CreatedByURL": "http://epicgames.com",
+	"DocsURL": "",
+	"MarketplaceURL": "",
+	"SupportURL": "",
+	"CanContainContent": true,
+	"IsBetaVersion": true,
+	"Installed": false,
+	"Modules": [
+		{
+			"Name": "ModelingToolsEditorMode",
+			"Type": "Editor",
+			"LoadingPhase": "Default",
+			"WhitelistPlatforms": [ "Win64" ]
+		}
+	],
+	"Plugins": [
+		{
+			"Name": "MeshModelingToolset",
+			"Enabled": true
+		}
+	]
+}
\ No newline at end of file
diff --git a/Engine/Plugins/Experimental/ModelingToolsEditorMode/Source/ModelingToolsEditorMode/Private/ModelingToolsActions.cpp b/Engine/Plugins/Experimental/ModelingToolsEditorMode/Source/ModelingToolsEditorMode/Private/ModelingToolsActions.cpp
new file mode 100644
--- /dev/null
+++ b/Engine/Plugins/Experimental/ModelingToolsEditorMode/Source/ModelingToolsEditorMode/Private/ModelingToolsActions.cpp
@@ -0,0 +1,45 @@
+// Copyright 1998-2019 Epic Games, Inc. All Rights Reserved.
+
+#include "ModelingToolsActions.h"
+#include "EditorStyleSet.h"
+#include "DynamicMeshSculptTool.h"
+#include "MeshInspectorTool.h"
+#include "DrawPolygonTool.h"
+#include "PolygonMeshDeformTool.h"
+#include "ShapeSprayTool.h"
+#include "MeshSpaceDeformerTool.h"
+#include "MeshSelectionTool.h"
+#include "TransformMeshesTool.h"
+
+
+#define LOCTEXT_NAMESPACE "ModelingToolsCommands"
+
+
+
+FModelingToolActionCommands::FModelingToolActionCommands() : 
+	TInteractiveToolCommands<FModelingToolActionCommands>(
+		"ModelingToolsEditMode", // Context name for fast lookup
+		NSLOCTEXT("Contexts", "ModelingToolsEditMode", "Modeling Tools Mode"), // Localized context name for displaying
+		NAME_None, // Parent
+		FEditorStyle::GetStyleSetName() // Icon Style Set
+	)
+{
+}
+
+
+void FModelingToolActionCommands::GetToolDefaultObjectList(TArray<UInteractiveTool*>& ToolCDOs)
+{
+	ToolCDOs.Add(GetMutableDefault<UDynamicMeshSculptTool>());
+	ToolCDOs.Add(GetMutableDefault<UMeshInspectorTool>());
+	ToolCDOs.Add(GetMutableDefault<UDrawPolygonTool>());
+	ToolCDOs.Add(GetMutableDefault<UPolygonMeshDeformTool>());
+	ToolCDOs.Add(GetMutableDefault<UMeshSpaceDeformerTool>());
+	ToolCDOs.Add(GetMutableDefault<UShapeSprayTool>());
+	ToolCDOs.Add(GetMutableDefault<UMeshSelectionTool>());
+	ToolCDOs.Add(GetMutableDefault<UTransformMeshesTool>());
+}
+
+
+
+#undef LOCTEXT_NAMESPACE
+
diff --git a/Engine/Plugins/Experimental/ModelingToolsEditorMode/Source/ModelingToolsEditorMode/Private/ModelingToolsEditorMode.cpp b/Engine/Plugins/Experimental/ModelingToolsEditorMode/Source/ModelingToolsEditorMode/Private/ModelingToolsEditorMode.cpp
new file mode 100644
--- /dev/null
+++ b/Engine/Plugins/Experimental/ModelingToolsEditorMode/Source/ModelingToolsEditorMode/Private/ModelingToolsEditorMode.cpp
@@ -0,0 +1,572 @@
+// Copyright 1998-2019 Epic Games, Inc. All Rights Reserved.
+
+#include "ModelingToolsEditorMode.h"
+#include "ModelingToolsEditorModeToolkit.h"
+#include "Toolkits/ToolkitManager.h"
+#include "Framework/Commands/UICommandList.h"
+#include "Framework/Application/SlateApplication.h"
+#include "Framework/MultiBox/MultiBoxBuilder.h"
+#include "EditorViewportClient.h"
+
+//#include "SingleClickTool.h"
+//#include "MeshSurfacePointTool.h"
+//#include "MeshVertexDragTool.h"
+#include "DynamicMeshSculptTool.h"
+#include "PolygonMeshDeformTool.h"
+#include "ConvertToPolygonsTool.h"
+#include "MakeMeshTestTool.h"
+#include "SmoothMeshTool.h"
+#include "RemeshMeshTool.h"
+#include "SimplifyMeshTool.h"
+#include "MeshInspectorTool.h"
+#include "WeldMeshEdgesTool.h"
+#include "DrawPolygonTool.h"
+#include "ShapeSprayTool.h"
+#include "MergeMeshesTool.h"
+#include "VoxelCSGMeshesTool.h"
+#include "PlaneCutTool.h"
+#include "PolygonOnMeshTool.h"
+#include "DisplaceMeshTool.h"
+#include "MeshSpaceDeformerTool.h"
+#include "EditNormalsTool.h"
+#include "TransformMeshesTool.h"
+#include "MeshSelectionTool.h"
+#include "UVProjectionTool.h"
+
+#include "ParameterizeMeshTool.h"
+
+#include "EditorModeManager.h"
+
+// viewport interaction support
+#include "ViewportInteractor.h"
+#include "ActorViewportTransformable.h"
+#include "ViewportWorldInteraction.h"
+#include "IViewportInteractionModule.h"
+#include "LevelEditor.h"
+#include "ILevelViewport.h"
+#include "SLevelViewport.h"
+
+#include "ModelingToolsActions.h"
+#include "ModelingToolsManagerActions.h"
+
+#define LOCTEXT_NAMESPACE "FModelingToolsEditorMode"
+
+
+//#define ENABLE_DEBUG_PRINTING
+//#define TOOLED_ENABLE_VIEWPORT_INTERACTION
+
+
+// if set to 1, then on mode initialization we include buttons for prototype modeling tools
+static TAutoConsoleVariable<int32> CVarEnablePrototypeModelingTools(
+	TEXT("modeling.EnablePrototypes"),
+	0,
+	TEXT("Enable unsupported Experimental prototype Modeling Tools"));
+
+
+
+const FEditorModeID FModelingToolsEditorMode::EM_ModelingToolsEditorModeId = TEXT("EM_ModelingToolsEditorMode");
+
+FModelingToolsEditorMode::FModelingToolsEditorMode()
+{
+	ToolsContext = nullptr;
+
+	UICommandList = MakeShareable(new FUICommandList);
+}
+
+FModelingToolsEditorMode::~FModelingToolsEditorMode()
+{
+	if (ToolsContext != nullptr)
+	{
+		ToolsContext->ShutdownContext();
+		ToolsContext = nullptr;
+	}
+}
+
+
+
+
+void FModelingToolsEditorMode::BuildModeToolbar(FToolBarBuilder& ToolbarBuilder)
+{
+	ToolbarBuilder.AddToolBarButton(FModelingToolsManagerCommands::Get().AcceptActiveTool);
+	ToolbarBuilder.AddToolBarButton(FModelingToolsManagerCommands::Get().CancelActiveTool);
+	ToolbarBuilder.AddToolBarButton(FModelingToolsManagerCommands::Get().CompleteActiveTool);
+
+	ToolbarBuilder.AddSeparator();
+
+	ToolbarBuilder.AddToolBarButton(FModelingToolsManagerCommands::Get().BeginMakePrimitiveTool);
+	ToolbarBuilder.AddToolBarButton(FModelingToolsManagerCommands::Get().BeginDrawPolygonTool);
+	if (CVarEnablePrototypeModelingTools.GetValueOnGameThread() > 0)
+	{
+		ToolbarBuilder.AddToolBarButton(FModelingToolsManagerCommands::Get().BeginShapeSprayTool);
+	}
+
+	ToolbarBuilder.AddSeparator();
+
+	ToolbarBuilder.AddToolBarButton(FModelingToolsManagerCommands::Get().BeginTransformMeshesTool);
+	ToolbarBuilder.AddToolBarButton(FModelingToolsManagerCommands::Get().BeginSculptMeshTool);
+	ToolbarBuilder.AddToolBarButton(FModelingToolsManagerCommands::Get().BeginPolyEditTool);
+	ToolbarBuilder.AddToolBarButton(FModelingToolsManagerCommands::Get().BeginSmoothMeshTool);
+	ToolbarBuilder.AddToolBarButton(FModelingToolsManagerCommands::Get().BeginDisplaceMeshTool);
+	if (CVarEnablePrototypeModelingTools.GetValueOnGameThread() > 0)
+	{
+		ToolbarBuilder.AddToolBarButton(FModelingToolsManagerCommands::Get().BeginMeshSpaceDeformerTool);
+	}
+
+	ToolbarBuilder.AddSeparator();
+
+	if (CVarEnablePrototypeModelingTools.GetValueOnGameThread() > 0)
+	{
+		ToolbarBuilder.AddToolBarButton(FModelingToolsManagerCommands::Get().BeginRemeshSculptMeshTool);
+	}
+	ToolbarBuilder.AddToolBarButton(FModelingToolsManagerCommands::Get().BeginMeshSelectionTool);
+	ToolbarBuilder.AddToolBarButton(FModelingToolsManagerCommands::Get().BeginPlaneCutTool);
+	ToolbarBuilder.AddToolBarButton(FModelingToolsManagerCommands::Get().BeginSimplifyMeshTool);
+	ToolbarBuilder.AddToolBarButton(FModelingToolsManagerCommands::Get().BeginRemeshMeshTool);
+	ToolbarBuilder.AddToolBarButton(FModelingToolsManagerCommands::Get().BeginUVProjectionTool);
+	ToolbarBuilder.AddToolBarButton(FModelingToolsManagerCommands::Get().BeginVoxelMergeTool);
+	ToolbarBuilder.AddToolBarButton(FModelingToolsManagerCommands::Get().BeginVoxelBooleanTool);
+	if (CVarEnablePrototypeModelingTools.GetValueOnGameThread() > 0)
+	{
+		ToolbarBuilder.AddToolBarButton(FModelingToolsManagerCommands::Get().BeginPolygonOnMeshTool);
+	}
+
+	ToolbarBuilder.AddSeparator();
+
+	ToolbarBuilder.AddToolBarButton(FModelingToolsManagerCommands::Get().BeginEditNormalsTool);
+	ToolbarBuilder.AddToolBarButton(FModelingToolsManagerCommands::Get().BeginWeldEdgesTool);
+	ToolbarBuilder.AddToolBarButton(FModelingToolsManagerCommands::Get().BeginMeshInspectorTool);
+	if (CVarEnablePrototypeModelingTools.GetValueOnGameThread() > 0)
+	{
+		ToolbarBuilder.AddToolBarButton(FModelingToolsManagerCommands::Get().BeginParameterizeMeshTool);
+		ToolbarBuilder.AddToolBarButton(FModelingToolsManagerCommands::Get().BeginPolyGroupsTool);
+	}
+
+}
+
+
+
+void FModelingToolsEditorMode::ActorSelectionChangeNotify()
+{
+}
+
+
+
+bool FModelingToolsEditorMode::ProcessEditDelete()
+{
+	return ToolsContext->ProcessEditDelete();
+}
+
+
+
+bool FModelingToolsEditorMode::AllowWidgetMove()
+{ 
+	return false; 
+}
+
+bool FModelingToolsEditorMode::ShouldDrawWidget() const
+{ 
+	// allow standard xform gizmo if we don't have an active tool
+	if (ToolsContext != nullptr)
+	{
+		return ToolsContext->ToolManager->HasAnyActiveTool() == false;
+	}
+	return true; 
+}
+
+bool FModelingToolsEditorMode::UsesTransformWidget() const
+{ 
+	return true; 
+}
+
+
+void FModelingToolsEditorMode::Tick(FEditorViewportClient* ViewportClient, float DeltaTime)
+{
+	FEdMode::Tick(ViewportClient, DeltaTime);
+
+	if (ToolsContext != nullptr)
+	{
+		ToolsContext->Tick(ViewportClient, DeltaTime);
+	}
+}
+
+
+
+
+void FModelingToolsEditorMode::Render(const FSceneView* View, FViewport* Viewport, FPrimitiveDrawInterface* PDI)
+{
+	FEdMode::Render(View, Viewport, PDI);
+
+	if (ToolsContext != nullptr)
+	{
+		ToolsContext->Render(View, Viewport, PDI);
+	}
+}
+
+
+
+
+
+bool FModelingToolsEditorMode::InputKey(FEditorViewportClient* ViewportClient, FViewport* Viewport, FKey Key, EInputEvent Event)
+{
+	// try hotkeys
+	if (Event != IE_Released)
+	{
+		if (UICommandList->ProcessCommandBindings(Key, FSlateApplication::Get().GetModifierKeys(), false/*Event == IE_Repeat*/))
+		{
+			return true;
+		}
+	}
+
+	bool bHandled = ToolsContext->InputKey(ViewportClient, Viewport, Key, Event);
+	if (bHandled == false)
+	{
+		bHandled = FEdMode::InputKey(ViewportClient, Viewport, Key, Event);
+	}
+	return bHandled;
+
+	//bool bHandled = FEdMode::InputKey(ViewportClient, Viewport, Key, Event);
+	//bHandled |= ToolsContext->InputKey(ViewportClient, Viewport, Key, Event);
+	//return bHandled;
+}
+
+
+bool FModelingToolsEditorMode::InputAxis(FEditorViewportClient* InViewportClient, FViewport* Viewport, int32 ControllerId, FKey Key, float Delta, float DeltaTime)
+{
+	// mouse axes: EKeys::MouseX, EKeys::MouseY, EKeys::MouseWheelAxis
+	return FEdMode::InputAxis(InViewportClient, Viewport, ControllerId, Key, Delta, DeltaTime);
+}
+
+
+
+
+
+bool FModelingToolsEditorMode::StartTracking(FEditorViewportClient* InViewportClient, FViewport* InViewport) 
+{ 
+	bool bHandled = FEdMode::StartTracking(InViewportClient, InViewport);
+#ifdef ENABLE_DEBUG_PRINTING
+	UE_LOG(LogTemp, Warning, TEXT("START TRACKING - base handled was %d"), (int)bHandled);
+#endif
+
+	bHandled |= ToolsContext->StartTracking(InViewportClient, InViewport);
+
+	return bHandled;
+}
+
+bool FModelingToolsEditorMode::CapturedMouseMove(FEditorViewportClient* InViewportClient, FViewport* InViewport, int32 InMouseX, int32 InMouseY)
+{
+	bool bHandled = ToolsContext->CapturedMouseMove(InViewportClient, InViewport, InMouseX, InMouseY);
+	return bHandled;
+}
+
+
+bool FModelingToolsEditorMode::EndTracking(FEditorViewportClient* InViewportClient, FViewport* InViewport)
+{
+	bool bHandled = ToolsContext->EndTracking(InViewportClient, InViewport);
+	return bHandled;
+}
+
+
+
+
+
+
+
+
+bool FModelingToolsEditorMode::ReceivedFocus(FEditorViewportClient* ViewportClient, FViewport* Viewport)
+{
+#ifdef ENABLE_DEBUG_PRINTING
+	UE_LOG(LogTemp, Warning, TEXT("RECEIVED FOCUS"));
+#endif
+
+	return false;
+}
+
+bool FModelingToolsEditorMode::LostFocus(FEditorViewportClient* ViewportClient, FViewport* Viewport)
+{
+#ifdef ENABLE_DEBUG_PRINTING
+	UE_LOG(LogTemp, Warning, TEXT("LOST FOCUS"));
+#endif
+
+	return false;
+}
+
+
+
+bool FModelingToolsEditorMode::MouseEnter(FEditorViewportClient* ViewportClient, FViewport* Viewport, int32 x, int32 y)
+{
+	bool bHandled = ToolsContext->MouseEnter(ViewportClient, Viewport, x, y);
+	return bHandled;
+}
+
+bool FModelingToolsEditorMode::MouseMove(FEditorViewportClient* ViewportClient, FViewport* Viewport, int32 x, int32 y)
+{
+	bool bHandled = ToolsContext->MouseMove(ViewportClient, Viewport, x, y);
+	return bHandled;
+}
+
+bool FModelingToolsEditorMode::MouseLeave(FEditorViewportClient* ViewportClient, FViewport* Viewport)
+{
+	bool bHandled = ToolsContext->MouseLeave(ViewportClient, Viewport);
+	return bHandled;
+}
+
+
+
+
+static FString MeshInspectorToolName("MeshInspectorTool");
+
+
+void FModelingToolsEditorMode::Enter()
+{
+	FEdMode::Enter();
+
+	// initialize FEdMode ToolsContext adapter
+
+	ToolsContext = NewObject<UEdModeInteractiveToolsContext>(GetTransientPackage(), TEXT("ToolsContext"), RF_Transient);
+	ToolsContext->InitializeContextFromEdMode(this);
+
+	if (!Toolkit.IsValid() && UsesToolkits())
+	{
+		Toolkit = MakeShareable(new FModelingToolsEditorModeToolkit);
+		Toolkit->Init(Owner->GetToolkitHost());
+
+
+		const FModelingToolsManagerCommands& ToolManagerCommands = FModelingToolsManagerCommands::Get();
+		const TSharedRef<FUICommandList>& CommandList = Toolkit->GetToolkitCommands();
+
+		CommandList->MapAction(
+			ToolManagerCommands.AcceptActiveTool,
+			FExecuteAction::CreateLambda([this]() { ToolsContext->EndTool(EToolShutdownType::Accept); }),
+			FCanExecuteAction::CreateLambda([this]() { return ToolsContext->CanAcceptActiveTool(); }),
+			FGetActionCheckState(),
+			FIsActionButtonVisible::CreateLambda([this]() {return ToolsContext->ActiveToolHasAccept(); }),
+			EUIActionRepeatMode::RepeatDisabled
+		);
+
+		CommandList->MapAction(
+			ToolManagerCommands.CancelActiveTool,
+			FExecuteAction::CreateLambda([this]() { ToolsContext->EndTool(EToolShutdownType::Cancel); }),
+			FCanExecuteAction::CreateLambda([this]() { return ToolsContext->CanCancelActiveTool(); }),
+			FGetActionCheckState(),
+			FIsActionButtonVisible::CreateLambda([this]() {return ToolsContext->ActiveToolHasAccept(); }),
+			EUIActionRepeatMode::RepeatDisabled
+		);
+
+		CommandList->MapAction(
+			ToolManagerCommands.CompleteActiveTool,
+			FExecuteAction::CreateLambda([this]() { ToolsContext->EndTool(EToolShutdownType::Completed); }),
+			FCanExecuteAction::CreateLambda([this]() { return ToolsContext->CanCompleteActiveTool(); }),
+			FGetActionCheckState(),
+			FIsActionButtonVisible::CreateLambda([this]() {return ToolsContext->CanCompleteActiveTool(); }),
+			EUIActionRepeatMode::RepeatDisabled
+		);
+
+	}
+
+	const FModelingToolsManagerCommands& ToolManagerCommands = FModelingToolsManagerCommands::Get();
+	const TSharedRef<FUICommandList>& CommandList = Toolkit->GetToolkitCommands();
+	auto RegisterToolFunc = [this, &CommandList](TSharedPtr<FUICommandInfo> UICommand, FString ToolIdentifier, UInteractiveToolBuilder* Builder)
+	{
+		ToolsContext->ToolManager->RegisterToolType(ToolIdentifier, Builder);
+		CommandList->MapAction( UICommand,
+			FExecuteAction::CreateLambda([this, ToolIdentifier]() { ToolsContext->StartTool(ToolIdentifier); }),
+			FCanExecuteAction::CreateLambda([this, ToolIdentifier]() { return ToolsContext->CanStartTool(ToolIdentifier); }));
+	};
+
+
+	// register tool set
+
+	//
+	// make shape tools
+	//
+	auto MakeMeshTestToolBuilder = NewObject<UMakeMeshTestToolBuilder>();
+	MakeMeshTestToolBuilder->AssetAPI = ToolsContext->GetAssetAPI();
+	RegisterToolFunc(ToolManagerCommands.BeginMakePrimitiveTool, TEXT("MakeMeshTestTool"), MakeMeshTestToolBuilder);
+
+	auto DrawPolygonToolBuilder = NewObject<UDrawPolygonToolBuilder>();
+	DrawPolygonToolBuilder->AssetAPI = ToolsContext->GetAssetAPI();
+	RegisterToolFunc(ToolManagerCommands.BeginDrawPolygonTool, TEXT("DrawPolygonTool"), DrawPolygonToolBuilder);
+
+	auto ShapeSprayToolBuilder = NewObject<UShapeSprayToolBuilder>();
+	ShapeSprayToolBuilder->AssetAPI = ToolsContext->GetAssetAPI();
+	RegisterToolFunc(ToolManagerCommands.BeginShapeSprayTool, TEXT("ShapeSprayTool"), ShapeSprayToolBuilder);
+
+
+	//
+	// vertex deform tools
+	// 
+
+	auto MoveVerticesToolBuilder = NewObject<UDynamicMeshSculptToolBuilder>();
+	MoveVerticesToolBuilder->bEnableRemeshing = false;
+	RegisterToolFunc(ToolManagerCommands.BeginSculptMeshTool, TEXT("MoveVerticesTool"), MoveVerticesToolBuilder);
+
+	RegisterToolFunc(ToolManagerCommands.BeginPolyEditTool, TEXT("PolygonMeshDeformTool"), NewObject<UPolygonMeshDeformToolBuilder>());
+	RegisterToolFunc(ToolManagerCommands.BeginSmoothMeshTool, TEXT("SmoothMeshTool"), NewObject<USmoothMeshToolBuilder>());
+	RegisterToolFunc(ToolManagerCommands.BeginDisplaceMeshTool, TEXT("DisplaceMeshTool"), NewObject<UDisplaceMeshToolBuilder>());
+	RegisterToolFunc(ToolManagerCommands.BeginMeshSpaceDeformerTool, TEXT("MeshSpaceDeformerTool"), NewObject<UMeshSpaceDeformerToolBuilder>());
+	RegisterToolFunc(ToolManagerCommands.BeginTransformMeshesTool, TEXT("TransformMeshesTool"), NewObject<UTransformMeshesToolBuilder>());
+
+	// edit tools
+
+
+	auto DynaSculptToolBuilder = NewObject<UDynamicMeshSculptToolBuilder>();
+	DynaSculptToolBuilder->bEnableRemeshing = true;
+	RegisterToolFunc(ToolManagerCommands.BeginRemeshSculptMeshTool, TEXT("DynaSculptTool"), DynaSculptToolBuilder);
+
+	RegisterToolFunc(ToolManagerCommands.BeginRemeshMeshTool, TEXT("RemeshMeshTool"), NewObject<URemeshMeshToolBuilder>());
+	RegisterToolFunc(ToolManagerCommands.BeginSimplifyMeshTool, TEXT("SimplifyMeshTool"), NewObject<USimplifyMeshToolBuilder>());
+
+
+	auto EditNormalsToolBuilder = NewObject<UEditNormalsToolBuilder>();
+	EditNormalsToolBuilder->AssetAPI = ToolsContext->GetAssetAPI();
+	RegisterToolFunc(ToolManagerCommands.BeginEditNormalsTool, TEXT("EditNormalsTool"), EditNormalsToolBuilder);
+
+	auto UVProjectionToolBuilder = NewObject<UUVProjectionToolBuilder>();
+	UVProjectionToolBuilder->AssetAPI = ToolsContext->GetAssetAPI();
+	RegisterToolFunc(ToolManagerCommands.BeginUVProjectionTool, TEXT("UVProjectionTool"), UVProjectionToolBuilder);
+
+	auto MergeMeshesToolBuilder = NewObject<UMergeMeshesToolBuilder>();
+	MergeMeshesToolBuilder->AssetAPI = ToolsContext->GetAssetAPI();
+	RegisterToolFunc(ToolManagerCommands.BeginVoxelMergeTool, TEXT("MergeMeshesTool"), MergeMeshesToolBuilder);
+
+	auto VoxelCSGMeshesToolBuilder = NewObject<UVoxelCSGMeshesToolBuilder>();
+	VoxelCSGMeshesToolBuilder->AssetAPI = ToolsContext->GetAssetAPI();
+	RegisterToolFunc(ToolManagerCommands.BeginVoxelBooleanTool, TEXT("VoxelCSGMeshesTool"), VoxelCSGMeshesToolBuilder);
+
+	auto PlaneCutToolBuilder = NewObject<UPlaneCutToolBuilder>();
+	PlaneCutToolBuilder->AssetAPI = ToolsContext->GetAssetAPI();
+	RegisterToolFunc(ToolManagerCommands.BeginPlaneCutTool, TEXT("PlaneCutTool"), PlaneCutToolBuilder);
+
+	auto PolygonOnMeshToolBuilder = NewObject<UPolygonOnMeshToolBuilder>();
+	PolygonOnMeshToolBuilder->AssetAPI = ToolsContext->GetAssetAPI();
+	RegisterToolFunc(ToolManagerCommands.BeginPolygonOnMeshTool, TEXT("PolygonOnMeshTool"), PolygonOnMeshToolBuilder);
+
+	auto ParameterizeMeshToolBuilder = NewObject<UParameterizeMeshToolBuilder>();
+	ParameterizeMeshToolBuilder->AssetAPI = ToolsContext->GetAssetAPI();
+	RegisterToolFunc(ToolManagerCommands.BeginParameterizeMeshTool, TEXT("ParameterizeMeshTool"), ParameterizeMeshToolBuilder);
+
+	auto MeshSelectionToolBuilder = NewObject<UMeshSelectionToolBuilder>();
+	MeshSelectionToolBuilder->AssetAPI = ToolsContext->GetAssetAPI();
+	RegisterToolFunc(ToolManagerCommands.BeginMeshSelectionTool, TEXT("MeshSelectionTool"), MeshSelectionToolBuilder);
+
+
+	// analysis tools
+
+	RegisterToolFunc(ToolManagerCommands.BeginMeshInspectorTool, TEXT("MeshInspectorTool"), NewObject<UMeshInspectorToolBuilder>());
+	RegisterToolFunc(ToolManagerCommands.BeginWeldEdgesTool, TEXT("WeldMeshEdgesTool"), NewObject<UWeldMeshEdgesToolBuilder>());
+	RegisterToolFunc(ToolManagerCommands.BeginPolyGroupsTool, TEXT("ConvertToPolygonsTool"), NewObject<UConvertToPolygonsToolBuilder>());
+
+
+
+	ToolsContext->ToolManager->SelectActiveToolType(EToolSide::Left, TEXT("DynaSculptTool"));
+
+	// listen for Tool start/end events to bind/unbind any hotkeys relevant to that Tool
+	ToolsContext->ToolManager->OnToolStarted.AddLambda([this](UInteractiveToolManager* Manager, UInteractiveTool* Tool)
+	{
+		const FModelingToolActionCommands& Commands = FModelingToolActionCommands::Get();
+		Commands.BindCommandsForCurrentTool(UICommandList, Tool);
+	});
+	ToolsContext->ToolManager->OnToolEnded.AddLambda([this](UInteractiveToolManager* Manager, UInteractiveTool* Tool)
+	{
+		const FModelingToolActionCommands& Commands = FModelingToolActionCommands::Get();
+		Commands.UnbindActiveCommands(UICommandList);
+	});
+
+
+#ifdef TOOLED_ENABLE_VIEWPORT_INTERACTION
+	///
+	// Viewport Interaction
+	///
+	UEditorWorldExtensionCollection* ExtensionCollection = GEditor->GetEditorWorldExtensionsManager()->GetEditorWorldExtensions(GetWorld());
+	check(ExtensionCollection != nullptr);
+	this->ViewportWorldInteraction = NewObject<UViewportWorldInteraction>(ExtensionCollection);
+	ExtensionCollection->AddExtension(this->ViewportWorldInteraction);
+		//Cast<UViewportWorldInteraction>(ExtensionCollection->AddExtension(UViewportWorldInteraction::StaticClass()));
+	check(ViewportWorldInteraction != nullptr);
+	//this->ViewportWorldInteraction->UseLegacyInteractions();
+	//this->ViewportWorldInteraction->AddMouseCursorInteractor();
+	this->ViewportWorldInteraction->SetUseInputPreprocessor(true);
+	this->ViewportWorldInteraction->SetGizmoHandleType(EGizmoHandleTypes::All);
+
+	// Set the current viewport.
+	{
+		const TSharedRef< ILevelEditor >& LevelEditor = FModuleManager::GetModuleChecked<FLevelEditorModule>("LevelEditor").GetFirstLevelEditor().ToSharedRef();
+
+		// Do we have an active perspective viewport that is valid for VR?  If so, go ahead and use that.
+		TSharedPtr<FEditorViewportClient> ViewportClient;
+		{
+			TSharedPtr<ILevelViewport> ActiveLevelViewport = LevelEditor->GetActiveViewportInterface();
+			if (ActiveLevelViewport.IsValid())
+			{
+				ViewportClient = StaticCastSharedRef<SLevelViewport>(ActiveLevelViewport->AsWidget())->GetViewportClient();
+			}
+		}
+
+		this->ViewportWorldInteraction->SetDefaultOptionalViewportClient(ViewportClient);
+	}
+#endif  // TOOLED_ENABLE_VIEWPORT_INTERACTION
+}
+
+
+
+
+void FModelingToolsEditorMode::Exit()
+{
+	ToolsContext->ShutdownContext();
+	ToolsContext = nullptr;
+
+	if (Toolkit.IsValid())
+	{
+		FToolkitManager::Get().CloseToolkit(Toolkit.ToSharedRef());
+		Toolkit.Reset();
+	}
+
+
+#ifdef TOOLED_ENABLE_VIEWPORT_INTERACTION
+	///
+	// Viewport Interaction
+	//
+	if (IViewportInteractionModule::IsAvailable())
+	{
+		if (ViewportWorldInteraction != nullptr)
+		{
+			ViewportWorldInteraction->ReleaseMouseCursorInteractor();
+
+			// Make sure gizmo is visible.  We may have hidden it
+			ViewportWorldInteraction->SetTransformGizmoVisible(true);
+
+			// Unregister mesh element transformer
+			//ViewportWorldInteraction->SetTransformer(nullptr);
+
+			UEditorWorldExtensionCollection* ExtensionCollection = GEditor->GetEditorWorldExtensionsManager()->GetEditorWorldExtensions(GetWorld());
+			if (ExtensionCollection != nullptr)
+			{
+				ExtensionCollection->RemoveExtension(ViewportWorldInteraction);
+			}
+
+			ViewportWorldInteraction = nullptr;
+		}
+	}
+#endif // TOOLED_ENABLE_VIEWPORT_INTERACTION
+
+
+	// Call base Exit method to ensure proper cleanup
+	FEdMode::Exit();
+}
+
+bool FModelingToolsEditorMode::UsesToolkits() const
+{
+	return true;
+}
+
+
+void FModelingToolsEditorMode::AddReferencedObjects(FReferenceCollector& Collector)
+{
+	Collector.AddReferencedObject(ToolsContext);
+}
+
+
+
+
+
+#undef LOCTEXT_NAMESPACE
\ No newline at end of file
diff --git a/Engine/Plugins/Experimental/ModelingToolsEditorMode/Source/ModelingToolsEditorMode/Private/ModelingToolsEditorModeModule.cpp b/Engine/Plugins/Experimental/ModelingToolsEditorMode/Source/ModelingToolsEditorMode/Private/ModelingToolsEditorModeModule.cpp
new file mode 100644
--- /dev/null
+++ b/Engine/Plugins/Experimental/ModelingToolsEditorMode/Source/ModelingToolsEditorMode/Private/ModelingToolsEditorModeModule.cpp
@@ -0,0 +1,51 @@
+// Copyright 1998-2019 Epic Games, Inc. All Rights Reserved.
+
+#include "ModelingToolsEditorModeModule.h"
+#include "ModelingToolsEditorMode.h"
+
+#include "ModelingToolsActions.h"
+#include "ModelingToolsManagerActions.h"
+#include "ModelingToolsEditorModeStyle.h"
+
+#define LOCTEXT_NAMESPACE "FModelingToolsEditorModeModule"
+
+void FModelingToolsEditorModeModule::StartupModule()
+{
+	FCoreDelegates::OnPostEngineInit.AddRaw(this, &FModelingToolsEditorModeModule::OnPostEngineInit);
+}
+
+void FModelingToolsEditorModeModule::ShutdownModule()
+{
+	FCoreDelegates::OnPostEngineInit.RemoveAll(this);
+
+	FModelingToolActionCommands::Unregister();
+	FModelingToolsManagerCommands::Unregister();
+
+	// Unregister slate style overrides
+	FModelingToolsEditorModeStyle::Shutdown();
+
+	// This function may be called during shutdown to clean up your module.  For modules that support dynamic reloading,
+	// we call this function before unloading the module.
+	FEditorModeRegistry::Get().UnregisterMode(FModelingToolsEditorMode::EM_ModelingToolsEditorModeId);
+}
+
+
+void FModelingToolsEditorModeModule::OnPostEngineInit()
+{
+	// Register slate style overrides
+	FModelingToolsEditorModeStyle::Initialize();
+
+	// This code will execute after your module is loaded into memory; the exact timing is specified in the .uplugin file per-module
+	FEditorModeRegistry::Get().RegisterMode<FModelingToolsEditorMode>(
+		FModelingToolsEditorMode::EM_ModelingToolsEditorModeId,
+		LOCTEXT("ModelingToolsEditorModeName", "ModelingToolsEditorMode"),
+		FSlateIcon(),
+		true);
+
+	FModelingToolActionCommands::Register();
+	FModelingToolsManagerCommands::Register();
+}
+
+#undef LOCTEXT_NAMESPACE
+	
+IMPLEMENT_MODULE(FModelingToolsEditorModeModule, ModelingToolsEditorMode)
\ No newline at end of file
diff --git a/Engine/Plugins/Experimental/ModelingToolsEditorMode/Source/ModelingToolsEditorMode/Private/ModelingToolsEditorModeStyle.cpp b/Engine/Plugins/Experimental/ModelingToolsEditorMode/Source/ModelingToolsEditorMode/Private/ModelingToolsEditorModeStyle.cpp
new file mode 100644
--- /dev/null
+++ b/Engine/Plugins/Experimental/ModelingToolsEditorMode/Source/ModelingToolsEditorMode/Private/ModelingToolsEditorModeStyle.cpp
@@ -0,0 +1,212 @@
+// Copyright 1998-2019 Epic Games, Inc. All Rights Reserved.
+
+#include "ModelingToolsEditorModeStyle.h"
+#include "Styling/SlateStyleRegistry.h"
+#include "Styling/SlateTypes.h"
+#include "Styling/CoreStyle.h"
+#include "EditorStyleSet.h"
+#include "Interfaces/IPluginManager.h"
+#include "SlateOptMacros.h"
+
+
+#define IMAGE_PLUGIN_BRUSH( RelativePath, ... ) FSlateImageBrush( FModelingToolsEditorModeStyle::InContent( RelativePath, ".png" ), __VA_ARGS__ )
+#define IMAGE_BRUSH(RelativePath, ...) FSlateImageBrush(StyleSet->RootToContentDir(RelativePath, TEXT(".png")), __VA_ARGS__)
+#define BOX_BRUSH(RelativePath, ...) FSlateBoxBrush(StyleSet->RootToContentDir(RelativePath, TEXT(".png")), __VA_ARGS__)
+#define DEFAULT_FONT(...) FCoreStyle::GetDefaultFontStyle(__VA_ARGS__)
+
+FString FModelingToolsEditorModeStyle::InContent(const FString& RelativePath, const ANSICHAR* Extension)
+{
+	static FString ContentDir = IPluginManager::Get().FindPlugin(TEXT("ModelingToolsEditorMode"))->GetContentDir();
+	return (ContentDir / RelativePath) + Extension;
+}
+
+TSharedPtr< FSlateStyleSet > FModelingToolsEditorModeStyle::StyleSet = nullptr;
+TSharedPtr< class ISlateStyle > FModelingToolsEditorModeStyle::Get() { return StyleSet; }
+
+FName FModelingToolsEditorModeStyle::GetStyleSetName()
+{
+	static FName ModelingToolsStyleName(TEXT("ModelingToolsStyle"));
+	return ModelingToolsStyleName;
+}
+
+BEGIN_SLATE_FUNCTION_BUILD_OPTIMIZATION
+
+void FModelingToolsEditorModeStyle::Initialize()
+{
+	// Const icon sizes
+	const FVector2D Icon8x8(8.0f, 8.0f);
+	const FVector2D Icon16x16(16.0f, 16.0f);
+	const FVector2D Icon20x20(20.0f, 20.0f);
+	const FVector2D Icon28x28(28.0f, 28.0f);
+	const FVector2D Icon40x40(40.0f, 40.0f);
+
+	// Only register once
+	if (StyleSet.IsValid())
+	{
+		return;
+	}
+
+	StyleSet = MakeShareable(new FSlateStyleSet(GetStyleSetName()));
+	StyleSet->SetContentRoot(FPaths::EngineContentDir() / TEXT("Editor/Slate"));
+	StyleSet->SetCoreContentRoot(FPaths::EngineContentDir() / TEXT("Slate"));
+
+	const FTextBlockStyle& NormalText = FEditorStyle::Get().GetWidgetStyle<FTextBlockStyle>("NormalText");
+
+	// Shared editors
+	//{
+	//	StyleSet->Set("Paper2D.Common.ViewportZoomTextStyle", FTextBlockStyle(NormalText)
+	//		.SetFont(DEFAULT_FONT("BoldCondensed", 16))
+	//	);
+
+	//	StyleSet->Set("Paper2D.Common.ViewportTitleTextStyle", FTextBlockStyle(NormalText)
+	//		.SetFont(DEFAULT_FONT("Regular", 18))
+	//		.SetColorAndOpacity(FLinearColor(1.0, 1.0f, 1.0f, 0.5f))
+	//	);
+
+	//	StyleSet->Set("Paper2D.Common.ViewportTitleBackground", new BOX_BRUSH("Old/Graph/GraphTitleBackground", FMargin(0)));
+	//}
+
+	// Tool Manager icons
+	{
+		// Accept/Cancel/Complete active tool
+		StyleSet->Set("ModelingToolsManagerCommands.CancelActiveTool", new IMAGE_PLUGIN_BRUSH("Icons/icon_ActiveTool_Cancel_40x", Icon28x28));
+		StyleSet->Set("ModelingToolsManagerCommands.CancelActiveTool.Small", new IMAGE_PLUGIN_BRUSH("Icons/icon_ActiveTool_Cancel_40x", Icon28x28));
+		StyleSet->Set("ModelingToolsManagerCommands.AcceptActiveTool", new IMAGE_PLUGIN_BRUSH("Icons/icon_ActiveTool_Accept_40x", Icon28x28));
+		StyleSet->Set("ModelingToolsManagerCommands.AcceptActiveTool.Small", new IMAGE_PLUGIN_BRUSH("Icons/icon_ActiveTool_Accept_40x", Icon28x28));
+		StyleSet->Set("ModelingToolsManagerCommands.CompleteActiveTool", new IMAGE_PLUGIN_BRUSH("Icons/icon_ActiveTool_Accept_40x", Icon28x28));
+		StyleSet->Set("ModelingToolsManagerCommands.CompleteActiveTool.Small", new IMAGE_PLUGIN_BRUSH("Icons/icon_ActiveTool_Accept_40x", Icon28x28));
+
+
+		StyleSet->Set("ModelingToolsManagerCommands.BeginMakePrimitiveTool", new IMAGE_PLUGIN_BRUSH("Icons/icon_Tool_MakePrimitive_40x", Icon28x28));
+		StyleSet->Set("ModelingToolsManagerCommands.BeginMakePrimitiveTool.Small", new IMAGE_PLUGIN_BRUSH("Icons/icon_Tool_MakePrimitive_40x", Icon28x28));
+		StyleSet->Set("ModelingToolsManagerCommands.BeginDrawPolygonTool", new IMAGE_PLUGIN_BRUSH("Icons/icon_Tool_DrawPolygon_40x", Icon28x28));
+		StyleSet->Set("ModelingToolsManagerCommands.BeginDrawPolygonTool.Small", new IMAGE_PLUGIN_BRUSH("Icons/icon_Tool_DrawPolygon_40x", Icon28x28));
+		StyleSet->Set("ModelingToolsManagerCommands.BeginShapeSprayTool", new IMAGE_PLUGIN_BRUSH("Icons/icon_Tool_ShapeSpray_40x", Icon28x28));
+		StyleSet->Set("ModelingToolsManagerCommands.BeginShapeSprayTool.Small", new IMAGE_PLUGIN_BRUSH("Icons/icon_Tool_ShapeSpray_40x", Icon28x28));
+
+		StyleSet->Set("ModelingToolsManagerCommands.BeginSmoothMeshTool", new IMAGE_PLUGIN_BRUSH("Icons/icon_Tool_Smooth_40x", Icon28x28));
+		StyleSet->Set("ModelingToolsManagerCommands.BeginSmoothMeshTool.Small", new IMAGE_PLUGIN_BRUSH("Icons/icon_Tool_Smooth_40x", Icon28x28));
+		StyleSet->Set("ModelingToolsManagerCommands.BeginSculptMeshTool", new IMAGE_PLUGIN_BRUSH("Icons/icon_Tool_Sculpt_40x", Icon28x28));
+		StyleSet->Set("ModelingToolsManagerCommands.BeginSculptMeshTool.Small", new IMAGE_PLUGIN_BRUSH("Icons/icon_Tool_Sculpt_40x", Icon28x28));
+		StyleSet->Set("ModelingToolsManagerCommands.BeginPolyEditTool", new IMAGE_PLUGIN_BRUSH("Icons/icon_Tool_PolyEdit_40x", Icon28x28));
+		StyleSet->Set("ModelingToolsManagerCommands.BeginPolyEditTool.Small", new IMAGE_PLUGIN_BRUSH("Icons/icon_Tool_PolyEdit_40x", Icon28x28));
+		StyleSet->Set("ModelingToolsManagerCommands.BeginDisplaceMeshTool", new IMAGE_PLUGIN_BRUSH("Icons/icon_Tool_Displace_40x", Icon28x28));
+		StyleSet->Set("ModelingToolsManagerCommands.BeginDisplaceMeshTool.Small", new IMAGE_PLUGIN_BRUSH("Icons/icon_Tool_Displace_40x", Icon28x28));
+		StyleSet->Set("ModelingToolsManagerCommands.BeginMeshSpaceDeformerTool", new IMAGE_PLUGIN_BRUSH("Icons/icon_Tool_Displace_40x", Icon28x28));
+		StyleSet->Set("ModelingToolsManagerCommands.BeginMeshSpaceDeformerTool.Small", new IMAGE_PLUGIN_BRUSH("Icons/icon_Tool_Displace_40x", Icon28x28));
+		StyleSet->Set("ModelingToolsManagerCommands.BeginTransformMeshesTool", new IMAGE_PLUGIN_BRUSH("Icons/icon_Tool_Transform_40x", Icon28x28));
+		StyleSet->Set("ModelingToolsManagerCommands.BeginTransformMeshesTool.Small", new IMAGE_PLUGIN_BRUSH("Icons/icon_Tool_Transform_40x", Icon28x28));
+
+		StyleSet->Set("ModelingToolsManagerCommands.BeginRemeshSculptMeshTool", new IMAGE_PLUGIN_BRUSH("Icons/icon_Tool_RemeshSculpt_40x", Icon28x28));
+		StyleSet->Set("ModelingToolsManagerCommands.BeginRemeshSculptMeshTool.Small", new IMAGE_PLUGIN_BRUSH("Icons/icon_Tool_RemeshSculpt_40x", Icon28x28));
+		StyleSet->Set("ModelingToolsManagerCommands.BeginRemeshMeshTool", new IMAGE_PLUGIN_BRUSH("Icons/icon_Tool_Remesh_40x", Icon28x28));
+		StyleSet->Set("ModelingToolsManagerCommands.BeginRemeshMeshTool.Small", new IMAGE_PLUGIN_BRUSH("Icons/icon_Tool_Remesh_40x", Icon28x28));
+		StyleSet->Set("ModelingToolsManagerCommands.BeginSimplifyMeshTool", new IMAGE_PLUGIN_BRUSH("Icons/icon_Tool_Simplify_40x", Icon28x28));
+		StyleSet->Set("ModelingToolsManagerCommands.BeginSimplifyMeshTool.Small", new IMAGE_PLUGIN_BRUSH("Icons/icon_Tool_Simplify_40x", Icon28x28));
+		StyleSet->Set("ModelingToolsManagerCommands.BeginEditNormalsTool", new IMAGE_PLUGIN_BRUSH("Icons/icon_Tool_EditNormals_40x", Icon28x28));
+		StyleSet->Set("ModelingToolsManagerCommands.BeginEditNormalsTool.Small", new IMAGE_PLUGIN_BRUSH("Icons/icon_Tool_EditNormals_40x", Icon28x28));
+		StyleSet->Set("ModelingToolsManagerCommands.BeginUVProjectionTool", new IMAGE_PLUGIN_BRUSH("Icons/icon_Tool_UVProjection_40x", Icon28x28));
+		StyleSet->Set("ModelingToolsManagerCommands.BeginUVProjectionTool.Small", new IMAGE_PLUGIN_BRUSH("Icons/icon_Tool_UVProjection_40x", Icon28x28));
+		StyleSet->Set("ModelingToolsManagerCommands.BeginVoxelMergeTool", new IMAGE_PLUGIN_BRUSH("Icons/icon_Tool_VoxMerge_40x", Icon28x28));
+		StyleSet->Set("ModelingToolsManagerCommands.BeginVoxelMergeTool.Small", new IMAGE_PLUGIN_BRUSH("Icons/icon_Tool_VoxMerge_40x", Icon28x28));
+		StyleSet->Set("ModelingToolsManagerCommands.BeginVoxelBooleanTool", new IMAGE_PLUGIN_BRUSH("Icons/icon_Tool_VoxBoolean_40x", Icon28x28));
+		StyleSet->Set("ModelingToolsManagerCommands.BeginVoxelBooleanTool.Small", new IMAGE_PLUGIN_BRUSH("Icons/icon_Tool_VoxBoolean_40x", Icon28x28));
+		StyleSet->Set("ModelingToolsManagerCommands.BeginPlaneCutTool", new IMAGE_PLUGIN_BRUSH("Icons/icon_Tool_PlaneCut_40x", Icon28x28));
+		StyleSet->Set("ModelingToolsManagerCommands.BeginPlaneCutTool.Small", new IMAGE_PLUGIN_BRUSH("Icons/icon_Tool_PlaneCut_40x", Icon28x28));
+		StyleSet->Set("ModelingToolsManagerCommands.BeginPolygonOnMeshTool", new IMAGE_PLUGIN_BRUSH("Icons/icon_Tool_PolygonOnMesh_40x", Icon28x28));
+		StyleSet->Set("ModelingToolsManagerCommands.BeginPolygonOnMeshTool.Small", new IMAGE_PLUGIN_BRUSH("Icons/icon_Tool_PolygonOnMesh_40x", Icon28x28));
+		StyleSet->Set("ModelingToolsManagerCommands.BeginMeshSelectionTool", new IMAGE_PLUGIN_BRUSH("Icons/icon_Tool_MeshSelection_40x", Icon28x28));
+		StyleSet->Set("ModelingToolsManagerCommands.BeginMeshSelectionTool.Small", new IMAGE_PLUGIN_BRUSH("Icons/icon_Tool_MeshSelection_40x", Icon28x28));
+
+		StyleSet->Set("ModelingToolsManagerCommands.BeginMeshInspectorTool", new IMAGE_PLUGIN_BRUSH("Icons/icon_Tool_MeshInspector_40x", Icon28x28));
+		StyleSet->Set("ModelingToolsManagerCommands.BeginMeshInspectorTool.Small", new IMAGE_PLUGIN_BRUSH("Icons/icon_Tool_MeshInspector_40x", Icon28x28));
+		StyleSet->Set("ModelingToolsManagerCommands.BeginParameterizeMeshTool", new IMAGE_PLUGIN_BRUSH("Icons/icon_Tool_UVGenerate_40x", Icon28x28));
+		StyleSet->Set("ModelingToolsManagerCommands.BeginParameterizeMeshTool.Small", new IMAGE_PLUGIN_BRUSH("Icons/icon_Tool_UVGenerate_40x", Icon28x28));
+		StyleSet->Set("ModelingToolsManagerCommands.BeginWeldEdgesTool", new IMAGE_PLUGIN_BRUSH("Icons/icon_Tool_WeldEdges_40x", Icon28x28));
+		StyleSet->Set("ModelingToolsManagerCommands.BeginWeldEdgesTool.Small", new IMAGE_PLUGIN_BRUSH("Icons/icon_Tool_WeldEdges_40x", Icon28x28));
+		StyleSet->Set("ModelingToolsManagerCommands.BeginPolyGroupsTool", new IMAGE_PLUGIN_BRUSH("Icons/icon_Tool_PolyGroups_40x", Icon28x28));
+		StyleSet->Set("ModelingToolsManagerCommands.BeginPolyGroupsTool.Small", new IMAGE_PLUGIN_BRUSH("Icons/icon_Tool_PolyGroups_40x", Icon28x28));
+
+
+		//StyleSet->Set("ModelingToolsManagerCommands.CancelActiveTool", new IMAGE_PLUGIN_BRUSH("Icons/icon_ActiveTool_Cancel_40x", Icon40x40));
+		//StyleSet->Set("ModelingToolsManagerCommands.CancelActiveTool.Small", new IMAGE_PLUGIN_BRUSH("Icons/icon_ActiveTool_Cancel_20x", Icon20x20));
+		//StyleSet->Set("ModelingToolsManagerCommands.AcceptActiveTool", new IMAGE_PLUGIN_BRUSH("Icons/icon_ActiveTool_Accept_40x", Icon40x40));
+		//StyleSet->Set("ModelingToolsManagerCommands.AcceptActiveTool.Small", new IMAGE_PLUGIN_BRUSH("Icons/icon_ActiveTool_Accept_20x", Icon20x20));
+		//StyleSet->Set("ModelingToolsManagerCommands.CompleteActiveTool", new IMAGE_PLUGIN_BRUSH("Icons/icon_ActiveTool_Accept_40x", Icon40x40));
+		//StyleSet->Set("ModelingToolsManagerCommands.CompleteActiveTool.Small", new IMAGE_PLUGIN_BRUSH("Icons/icon_ActiveTool_Accept_20x", Icon20x20));
+
+
+		//StyleSet->Set("ModelingToolsManagerCommands.BeginMakePrimitiveTool", new IMAGE_PLUGIN_BRUSH("Icons/icon_Tool_MakePrimitive_40x", Icon40x40));
+		//StyleSet->Set("ModelingToolsManagerCommands.BeginMakePrimitiveTool.Small", new IMAGE_PLUGIN_BRUSH("Icons/icon_Tool_MakePrimitive_20x", Icon20x20));
+		//StyleSet->Set("ModelingToolsManagerCommands.BeginDrawPolygonTool", new IMAGE_PLUGIN_BRUSH("Icons/icon_Tool_DrawPolygon_40x", Icon40x40));
+		//StyleSet->Set("ModelingToolsManagerCommands.BeginDrawPolygonTool.Small", new IMAGE_PLUGIN_BRUSH("Icons/icon_Tool_DrawPolygon_20x", Icon20x20));
+		//StyleSet->Set("ModelingToolsManagerCommands.BeginShapeSprayTool", new IMAGE_PLUGIN_BRUSH("Icons/icon_Tool_ShapeSpray_40x", Icon40x40));
+		//StyleSet->Set("ModelingToolsManagerCommands.BeginShapeSprayTool.Small", new IMAGE_PLUGIN_BRUSH("Icons/icon_Tool_ShapeSpray_20x", Icon20x20));
+
+		//StyleSet->Set("ModelingToolsManagerCommands.BeginSmoothMeshTool", new IMAGE_PLUGIN_BRUSH("Icons/icon_Tool_Smooth_40x", Icon40x40));
+		//StyleSet->Set("ModelingToolsManagerCommands.BeginSmoothMeshTool.Small", new IMAGE_PLUGIN_BRUSH("Icons/icon_Tool_Smooth_20x", Icon20x20));
+		//StyleSet->Set("ModelingToolsManagerCommands.BeginSculptMeshTool", new IMAGE_PLUGIN_BRUSH("Icons/icon_Tool_Sculpt_40x", Icon40x40));
+		//StyleSet->Set("ModelingToolsManagerCommands.BeginSculptMeshTool.Small", new IMAGE_PLUGIN_BRUSH("Icons/icon_Tool_Sculpt_20x", Icon20x20));
+		//StyleSet->Set("ModelingToolsManagerCommands.BeginPolyEditTool", new IMAGE_PLUGIN_BRUSH("Icons/icon_Tool_PolyEdit_40x", Icon40x40));
+		//StyleSet->Set("ModelingToolsManagerCommands.BeginPolyEditTool.Small", new IMAGE_PLUGIN_BRUSH("Icons/icon_Tool_PolyEdit_20x", Icon20x20));
+
+		//StyleSet->Set("ModelingToolsManagerCommands.BeginRemeshSculptMeshTool", new IMAGE_PLUGIN_BRUSH("Icons/icon_Tool_RemeshSculpt_40x", Icon40x40));
+		//StyleSet->Set("ModelingToolsManagerCommands.BeginRemeshSculptMeshTool.Small", new IMAGE_PLUGIN_BRUSH("Icons/icon_Tool_RemeshSculpt_20x", Icon20x20));
+		//StyleSet->Set("ModelingToolsManagerCommands.BeginRemeshMeshTool", new IMAGE_PLUGIN_BRUSH("Icons/icon_Tool_Remesh_40x", Icon40x40));
+		//StyleSet->Set("ModelingToolsManagerCommands.BeginRemeshMeshTool.Small", new IMAGE_PLUGIN_BRUSH("Icons/icon_Tool_Remesh_20x", Icon20x20));
+		//StyleSet->Set("ModelingToolsManagerCommands.BeginSimplifyMeshTool", new IMAGE_PLUGIN_BRUSH("Icons/icon_Tool_Simplify_40x", Icon40x40));
+		//StyleSet->Set("ModelingToolsManagerCommands.BeginSimplifyMeshTool.Small", new IMAGE_PLUGIN_BRUSH("Icons/icon_Tool_Simplify_20x", Icon20x20));
+
+
+		//StyleSet->Set("ModelingToolsManagerCommands.BeginMeshInspectorTool", new IMAGE_PLUGIN_BRUSH("Icons/icon_Tool_MeshInspector_40x", Icon40x40));
+		//StyleSet->Set("ModelingToolsManagerCommands.BeginMeshInspectorTool.Small", new IMAGE_PLUGIN_BRUSH("Icons/icon_Tool_MeshInspector_20x", Icon20x20));
+		//StyleSet->Set("ModelingToolsManagerCommands.BeginWeldEdgesTool", new IMAGE_PLUGIN_BRUSH("Icons/icon_Tool_WeldEdges_40x", Icon40x40));
+		//StyleSet->Set("ModelingToolsManagerCommands.BeginWeldEdgesTool.Small", new IMAGE_PLUGIN_BRUSH("Icons/icon_Tool_WeldEdges_20x", Icon20x20));
+		//StyleSet->Set("ModelingToolsManagerCommands.BeginPolyGroupsTool", new IMAGE_PLUGIN_BRUSH("Icons/icon_Tool_PolyGroups_40x", Icon40x40));
+		//StyleSet->Set("ModelingToolsManagerCommands.BeginPolyGroupsTool.Small", new IMAGE_PLUGIN_BRUSH("Icons/icon_Tool_PolyGroups_20x", Icon20x20));
+
+
+
+
+		//const FLinearColor LayerSelectionColor = FLinearColor(0.13f, 0.70f, 1.00f);
+
+		//// Selection color for the active row should be ??? to align with the editor viewport.
+		//const FTableRowStyle& NormalTableRowStyle = FEditorStyle::Get().GetWidgetStyle<FTableRowStyle>("TableView.Row");
+
+		//StyleSet->Set("TileMapEditor.LayerBrowser.TableViewRow",
+		//	FTableRowStyle(NormalTableRowStyle)
+		//	.SetActiveBrush(IMAGE_BRUSH("Common/Selection", Icon8x8, LayerSelectionColor))
+		//	.SetActiveHoveredBrush(IMAGE_BRUSH("Common/Selection", Icon8x8, LayerSelectionColor))
+		//	.SetInactiveBrush(IMAGE_BRUSH("Common/Selection", Icon8x8, LayerSelectionColor))
+		//	.SetInactiveHoveredBrush(IMAGE_BRUSH("Common/Selection", Icon8x8, LayerSelectionColor))
+		//);
+
+		//StyleSet->Set("TileMapEditor.LayerBrowser.SelectionColor", LayerSelectionColor);
+
+		//StyleSet->Set("TileMapEditor.TileSetPalette.NothingSelectedText", FTextBlockStyle(NormalText)
+		//	.SetFont(DEFAULT_FONT("BoldCondensed", 18))
+		//	.SetColorAndOpacity(FLinearColor(0.8, 0.8f, 0.0f, 0.8f))
+		//	.SetShadowOffset(FVector2D(1.0f, 1.0f))
+		//	.SetShadowColorAndOpacity(FLinearColor(0.0f, 0.0f, 0.0f, 0.9f))
+		//);
+	}
+
+	FSlateStyleRegistry::RegisterSlateStyle(*StyleSet.Get());
+};
+
+END_SLATE_FUNCTION_BUILD_OPTIMIZATION
+
+#undef IMAGE_PLUGIN_BRUSH
+#undef IMAGE_BRUSH
+#undef BOX_BRUSH
+#undef DEFAULT_FONT
+
+void FModelingToolsEditorModeStyle::Shutdown()
+{
+	if (StyleSet.IsValid())
+	{
+		FSlateStyleRegistry::UnRegisterSlateStyle(*StyleSet.Get());
+		ensure(StyleSet.IsUnique());
+		StyleSet.Reset();
+	}
+}
diff --git a/Engine/Plugins/Experimental/ModelingToolsEditorMode/Source/ModelingToolsEditorMode/Private/ModelingToolsEditorModeToolkit.cpp b/Engine/Plugins/Experimental/ModelingToolsEditorMode/Source/ModelingToolsEditorMode/Private/ModelingToolsEditorModeToolkit.cpp
new file mode 100644
--- /dev/null
+++ b/Engine/Plugins/Experimental/ModelingToolsEditorMode/Source/ModelingToolsEditorMode/Private/ModelingToolsEditorModeToolkit.cpp
@@ -0,0 +1,152 @@
+// Copyright 1998-2019 Epic Games, Inc. All Rights Reserved.
+
+#include "ModelingToolsEditorModeToolkit.h"
+#include "ModelingToolsEditorMode.h"
+#include "Engine/Selection.h"
+
+#include "Modules/ModuleManager.h"
+#include "PropertyEditorModule.h"
+#include "IDetailsView.h"
+
+#include "Widgets/Input/SButton.h"
+#include "Widgets/Text/STextBlock.h"
+
+#include "EditorModeManager.h"
+
+#define LOCTEXT_NAMESPACE "FModelingToolsEditorModeToolkit"
+
+FModelingToolsEditorModeToolkit::FModelingToolsEditorModeToolkit()
+{
+}
+
+
+TSharedRef<SButton> FModelingToolsEditorModeToolkit::MakeToolButton(const FText& ButtonLabel, const FString& ToolIdentifier)
+{
+	return SNew(SButton).Text(ButtonLabel)
+		.OnClicked_Lambda([this, ToolIdentifier]() { GetToolsContext()->StartTool(ToolIdentifier); return FReply::Handled(); })
+		.IsEnabled_Lambda([this, ToolIdentifier]() { return GetToolsContext()->CanStartTool(ToolIdentifier); });
+}
+
+SVerticalBox::FSlot& FModelingToolsEditorModeToolkit::MakeToolButtonSlotV(const FText& ButtonLabel, const FString& ToolIdentifier)
+{
+	return SVerticalBox::Slot()
+		.HAlign(HAlign_Center)
+		.AutoHeight()
+		[MakeToolButton(ButtonLabel, ToolIdentifier)];
+}
+
+SHorizontalBox::FSlot& FModelingToolsEditorModeToolkit::MakeToolButtonSlotH(const FText& ButtonLabel, const FString& ToolIdentifier)
+{
+	return SHorizontalBox::Slot()
+		.HAlign(HAlign_Center)
+		.AutoWidth()
+		[MakeToolButton(ButtonLabel, ToolIdentifier)];
+}
+
+
+SVerticalBox::FSlot& FModelingToolsEditorModeToolkit::MakeSetToolLabelV(const FText& LabelText)
+{
+	return SVerticalBox::Slot()
+		.AutoHeight()
+		.HAlign(HAlign_Center)
+		.Padding(5)
+		[
+			SNew(STextBlock)
+			.AutoWrapText(true)
+			.Text(LabelText)
+		];
+}
+
+
+
+void FModelingToolsEditorModeToolkit::Init(const TSharedPtr<IToolkitHost>& InitToolkitHost)
+{
+
+	FPropertyEditorModule& PropertyEditorModule = FModuleManager::GetModuleChecked<FPropertyEditorModule>("PropertyEditor");
+
+	FDetailsViewArgs DetailsViewArgs(
+		/*bUpdateFromSelection=*/ false,
+		/*bLockable=*/ false,
+		/*bAllowSearch=*/ false,
+		FDetailsViewArgs::HideNameArea,
+		/*bHideSelectionTip=*/ true,
+		/*InNotifyHook=*/ nullptr,
+		/*InSearchInitialKeyFocus=*/ false,
+		/*InViewIdentifier=*/ NAME_None);
+	DetailsViewArgs.DefaultsOnlyVisibility = EEditDefaultsOnlyNodeVisibility::Automatic;
+	DetailsViewArgs.bShowOptions = false;
+	DetailsViewArgs.bAllowMultipleTopLevelObjects = true;
+
+	DetailsView = PropertyEditorModule.CreateDetailView(DetailsViewArgs);
+
+
+	ToolHeaderLabel = SNew(STextBlock).AutoWrapText(true);
+	ToolHeaderLabel->SetText(LOCTEXT("SelectToolLabel", "Select a Tool from the Toolbar"));
+	ToolHeaderLabel->SetJustification(ETextJustify::Center);
+
+	SAssignNew(ToolkitWidget, SBorder)
+		.HAlign(HAlign_Fill)
+		.Padding(4)
+		[
+			SNew(SVerticalBox)
+
+			+ SVerticalBox::Slot().AutoHeight().HAlign(HAlign_Fill).Padding(5)
+				[ToolHeaderLabel->AsShared()]
+
+			+ SVerticalBox::Slot().HAlign(HAlign_Fill).AutoHeight().MaxHeight(500.0f)
+				[
+					DetailsView->AsShared()
+				]
+
+		];
+		
+	FModeToolkit::Init(InitToolkitHost);
+
+
+
+	GetToolsEditorMode()->GetToolManager()->OnToolStarted.AddLambda([this](UInteractiveToolManager* Manager, UInteractiveTool* Tool)
+	{
+		// Update properties panel
+		UInteractiveTool* CurTool = GetToolsEditorMode()->GetToolManager()->GetActiveTool(EToolSide::Left);
+		DetailsView->SetObjects(CurTool->GetToolProperties());
+
+		ToolHeaderLabel->SetText(CurTool->GetClass()->GetDisplayNameText());
+		//ToolHeaderLabel->SetText(FString("(Tool Name Here)"));
+	});
+	GetToolsEditorMode()->GetToolManager()->OnToolEnded.AddLambda([this](UInteractiveToolManager* Manager, UInteractiveTool* Tool)
+	{
+		DetailsView->SetObject(nullptr);
+		ToolHeaderLabel->SetText(LOCTEXT("SelectToolLabel", "Select a Tool from the Toolbar"));
+	});
+
+}
+
+
+
+
+FName FModelingToolsEditorModeToolkit::GetToolkitFName() const
+{
+	return FName("ModelingToolsEditorMode");
+}
+
+FText FModelingToolsEditorModeToolkit::GetBaseToolkitName() const
+{
+	return NSLOCTEXT("ModelingToolsEditorModeToolkit", "DisplayName", "ModelingToolsEditorMode Tool");
+}
+
+class FEdMode* FModelingToolsEditorModeToolkit::GetEditorMode() const
+{
+	return GLevelEditorModeTools().GetActiveMode(FModelingToolsEditorMode::EM_ModelingToolsEditorModeId);
+}
+
+FModelingToolsEditorMode * FModelingToolsEditorModeToolkit::GetToolsEditorMode() const
+{
+	return (FModelingToolsEditorMode *)GetEditorMode();
+}
+
+UEdModeInteractiveToolsContext* FModelingToolsEditorModeToolkit::GetToolsContext() const
+{
+	return GetToolsEditorMode()->GetToolsContext();
+}
+
+#undef LOCTEXT_NAMESPACE
diff --git a/Engine/Plugins/Experimental/ModelingToolsEditorMode/Source/ModelingToolsEditorMode/Private/ModelingToolsManagerActions.cpp b/Engine/Plugins/Experimental/ModelingToolsEditorMode/Source/ModelingToolsEditorMode/Private/ModelingToolsManagerActions.cpp
new file mode 100644
--- /dev/null
+++ b/Engine/Plugins/Experimental/ModelingToolsEditorMode/Source/ModelingToolsEditorMode/Private/ModelingToolsManagerActions.cpp
@@ -0,0 +1,59 @@
+// Copyright 1998-2019 Epic Games, Inc. All Rights Reserved.
+
+#include "ModelingToolsManagerActions.h"
+#include "EditorStyleSet.h"
+#include "ModelingToolsEditorModeStyle.h"
+
+#define LOCTEXT_NAMESPACE "ModelingToolsManagerCommands"
+
+
+
+FModelingToolsManagerCommands::FModelingToolsManagerCommands() :
+	TCommands<FModelingToolsManagerCommands>(
+		"ModelingToolsManagerCommands", // Context name for fast lookup
+		NSLOCTEXT("Contexts", "ModelingToolsEditMode", "Modeling Tools Mode"), // Localized context name for displaying
+		NAME_None, // Parent
+		FModelingToolsEditorModeStyle::Get()->GetStyleSetName() // Icon Style Set
+		)
+{
+}
+
+
+void FModelingToolsManagerCommands::RegisterCommands()
+{
+	UI_COMMAND(BeginMakePrimitiveTool, "Primitive", "Start the Make Primitive Tool", EUserInterfaceActionType::Button, FInputChord());
+	UI_COMMAND(BeginDrawPolygonTool, "Polygon", "Start the Draw Polygon Tool", EUserInterfaceActionType::Button, FInputChord());
+	UI_COMMAND(BeginShapeSprayTool, "ShapeSpray", "Start the Shape Spray Tool", EUserInterfaceActionType::Button, FInputChord());
+
+	UI_COMMAND(BeginSculptMeshTool, "Sculpt", "Start the Sculpt Mesh Tool", EUserInterfaceActionType::Button, FInputChord());
+	UI_COMMAND(BeginPolyEditTool, "PolyEdit", "Start the PolyEdit Tool", EUserInterfaceActionType::Button, FInputChord());
+	UI_COMMAND(BeginSmoothMeshTool, "Smooth", "Start the Smooth Mesh Tool", EUserInterfaceActionType::Button, FInputChord());
+	UI_COMMAND(BeginDisplaceMeshTool, "Displace", "Start the Displace Mesh Tool", EUserInterfaceActionType::Button, FInputChord());
+	UI_COMMAND(BeginMeshSpaceDeformerTool, "SpaceDeform", "Start the Mesh Space Deformer Tool", EUserInterfaceActionType::Button, FInputChord());
+	UI_COMMAND(BeginTransformMeshesTool, "Transform", "Start the Transform Tool", EUserInterfaceActionType::Button, FInputChord());
+
+	UI_COMMAND(BeginRemeshSculptMeshTool, "DynaSculpt", "Start the Liquid Sculpt Mesh Tool", EUserInterfaceActionType::Button, FInputChord());
+	UI_COMMAND(BeginRemeshMeshTool, "Remesh", "Start the Remesh Tool", EUserInterfaceActionType::Button, FInputChord());
+	UI_COMMAND(BeginSimplifyMeshTool, "Simplify", "Start the Simplify Mesh Tool", EUserInterfaceActionType::Button, FInputChord());
+	UI_COMMAND(BeginEditNormalsTool, "Normals", "Start the Edit Normals Tool", EUserInterfaceActionType::Button, FInputChord());
+	UI_COMMAND(BeginUVProjectionTool, "UVProjection", "Start the UV Projection Tool", EUserInterfaceActionType::Button, FInputChord());
+	UI_COMMAND(BeginVoxelMergeTool, "VoxMerge", "Start the Voxel Merge Tool", EUserInterfaceActionType::Button, FInputChord());
+	UI_COMMAND(BeginVoxelBooleanTool, "VoxBoolean", "Start the Voxel Boolean Tool", EUserInterfaceActionType::Button, FInputChord());
+	UI_COMMAND(BeginPlaneCutTool, "PlaneCut", "Start the Plane Cut Tool", EUserInterfaceActionType::Button, FInputChord());
+	UI_COMMAND(BeginPolygonOnMeshTool, "PolygonOnMesh", "Start the Polygon On Mesh Tool", EUserInterfaceActionType::Button, FInputChord());
+	UI_COMMAND(BeginParameterizeMeshTool, "UVGenerate", "Start the Mesh Parameterization Tool", EUserInterfaceActionType::Button, FInputChord());
+	UI_COMMAND(BeginMeshSelectionTool, "Select", "Start the Mesh Selection Tool", EUserInterfaceActionType::Button, FInputChord());
+
+	UI_COMMAND(BeginMeshInspectorTool, "Inspector", "Start the Mesh Inspector Tool", EUserInterfaceActionType::Button, FInputChord());
+	UI_COMMAND(BeginWeldEdgesTool, "Weld Edges", "Start the Weld Edges Tool", EUserInterfaceActionType::Button, FInputChord());
+	UI_COMMAND(BeginPolyGroupsTool, "PolyGroups", "Start the Poly Groups Tool", EUserInterfaceActionType::Button, FInputChord());
+
+	UI_COMMAND(AcceptActiveTool, "Accept", "Accept the active tool", EUserInterfaceActionType::Button, FInputChord());
+	UI_COMMAND(CancelActiveTool, "Cancel", "Cancel the active tool", EUserInterfaceActionType::Button, FInputChord());
+	UI_COMMAND(CompleteActiveTool, "Complete", "Complete the active tool", EUserInterfaceActionType::Button, FInputChord());
+}
+
+
+
+
+#undef LOCTEXT_NAMESPACE
diff --git a/Engine/Plugins/Experimental/ModelingToolsEditorMode/Source/ModelingToolsEditorMode/Public/ModelingToolsEditorMode.h b/Engine/Plugins/Experimental/ModelingToolsEditorMode/Source/ModelingToolsEditorMode/Public/ModelingToolsEditorMode.h
new file mode 100644
--- /dev/null
+++ b/Engine/Plugins/Experimental/ModelingToolsEditorMode/Source/ModelingToolsEditorMode/Public/ModelingToolsEditorMode.h
@@ -0,0 +1,159 @@
+// Copyright 1998-2019 Epic Games, Inc. All Rights Reserved.
+
+#pragma once
+
+#include "CoreMinimal.h"
+#include "EdMode.h"
+
+#include "InputState.h"
+#include "InteractiveToolManager.h"
+#include "EdModeInteractiveToolsContext.h"
+
+
+class FEditorComponentSourceFactory;
+class FEditorToolAssetAPI;
+class FUICommandList;
+
+class FModelingToolsEditorMode : public FEdMode
+{
+public:
+	const static FEditorModeID EM_ModelingToolsEditorModeId;
+public:
+	FModelingToolsEditorMode();
+	virtual ~FModelingToolsEditorMode();
+
+	////////////////
+	// FEdMode interface
+	////////////////
+
+	virtual void Tick(FEditorViewportClient* ViewportClient, float DeltaTime) override;
+	virtual void Render(const FSceneView* View, FViewport* Viewport, FPrimitiveDrawInterface* PDI) override;
+	virtual bool UsesToolkits() const override;
+
+	virtual void AddReferencedObjects(FReferenceCollector& Collector) override;
+
+	// these disable the standard gizmo, which is probably want we want in
+	// these tools as we can't hit-test the standard gizmo...
+	virtual bool AllowWidgetMove() override;
+	virtual bool ShouldDrawWidget() const override;
+	virtual bool UsesTransformWidget() const override;
+
+	virtual void ActorSelectionChangeNotify() override;
+
+	virtual bool ProcessEditDelete();
+
+	virtual void BuildModeToolbar(class FToolBarBuilder& ToolbarBuilder) override;
+
+
+	/*
+	 * focus events
+	 */
+
+	// called when we "start" this editor mode (ie switch to this tab)
+	virtual void Enter() override;
+
+	// called when we "end" this editor mode (ie switch to another tab)
+	virtual void Exit() override;
+
+	// called when viewport window is focused
+	virtual bool ReceivedFocus(FEditorViewportClient* ViewportClient, FViewport* Viewport) override;
+
+	// called when viewport window loses focus (ie when some other window is focused)
+	// *not* called when editor is backgrounded, but is called when editor is minimized
+	virtual bool LostFocus(FEditorViewportClient* ViewportClient, FViewport* Viewport) override;
+
+
+	/*
+	 * Mouse position events - These are called when no mouse button is down
+	 */
+
+	// called when mouse moves over viewport window
+	virtual bool MouseEnter(FEditorViewportClient* ViewportClient, FViewport* Viewport, int32 x, int32 y) override;
+
+	// called when mouse leaves viewport window
+	virtual bool MouseLeave(FEditorViewportClient* ViewportClient, FViewport* Viewport) override;
+
+	// called on any mouse-move event. *not* called during tracking/capturing, eg if any button is down
+	virtual bool MouseMove(FEditorViewportClient* ViewportClient, FViewport* Viewport, int32 x, int32 y) override;
+
+
+	/*
+	 * Input Button/Axis Events & Mouse Capture
+	 *
+	 *   event sequence for left mouse down/capture/up is:
+	 *     - InputKey( EKeys::LeftMouseButton, IE_PRESSED )
+	 *     - StartTracking()
+	 *     - CapturedMouseMove() and InputAxis() (repeated for each mouse-move)
+	 *     - Inputkey( EKeys::LeftMouseButton, IE_RELEASED )
+	 *     - EndTracking()
+	 *
+	 *   for doubleclick, we get one of the above sequence, then a second
+	 *   where instead of the first IE_PRESSED we get a IE_DoubleClick
+	 *
+	 *   for mouse wheel, we get following sequence
+	 *     - InputKey ( EKeys::MouseScrollUp / Down, IE_PRESSED )
+	 *     - InputAxis() for wheel move
+	 *     - InputKey ( EKeys::MouseScrollUp / Down, IE_RELEASED )
+	 *   it appears that there will only ever be *one* InputAxis() between the pressed/released sequenece
+	 *
+	 *   Note that this wheel event sequence can happen *during* a 
+	 *   middle-mouse tracking sequence on a wheel-button (nice!)
+	 */
+
+
+	// This is not just called for keyboard keys, it is also called for mouse down/up events!
+	// Eg for left-press we get Key = EKeys::LeftMouseButton and Event = IE_Pressed
+	// Return value indicates "handled'. If we return true for mouse press events then
+	// the StartTracking/CapturedMouseMove/EndTracking sequence is *not* called (but we
+	// also don't get MouseMove(), so the mouse-movement events appear to be lost?)
+	virtual bool InputKey(FEditorViewportClient* ViewportClient, FViewport* Viewport, FKey Key, EInputEvent Event) override;
+
+
+	// Called for 1D Axis movements - EKeys::MouseX, EKeys::MouseY, and EKeys::MouseWheelAxis
+	// Called even if we return true from InputKey which otherwise blocks tracking. 
+	// Return value indicates "handled" but has no effect on CapturedMouseMove events (ie we still get them)
+	virtual bool InputAxis(FEditorViewportClient* InViewportClient, FViewport* Viewport, int32 ControllerId, FKey Key, float Delta, float DeltaTime);
+
+
+	// called on mouse-down. return value indicates whether this was "handled" but
+	// does not mean we get exclusive capture events!
+	virtual bool StartTracking(FEditorViewportClient* InViewportClient, FViewport* InViewport) override;
+
+	// called during mouse-down mouse-move. Always called, return value is used to indicate whether
+	// we "handled" this mouse move (see FEditorModeTools::CapturedMouseMove) but does not pre-empt
+	// other editor modes from seeing this move event...
+	virtual bool CapturedMouseMove(FEditorViewportClient* InViewportClient, FViewport* InViewport, int32 InMouseX, int32 InMouseY) override;
+
+	// always called on mouse-up
+	virtual bool EndTracking(FEditorViewportClient* InViewportClient, FViewport* InViewport) override;
+
+
+	//////////////////
+	// End of FEdMode interface
+	//////////////////
+
+
+public:
+	virtual UEdModeInteractiveToolsContext* GetToolsContext() const
+	{
+		return ToolsContext;
+	}
+
+	virtual UInteractiveToolManager* GetToolManager() const
+	{ 
+		return ToolsContext->ToolManager;
+	}
+
+protected:
+
+	UEdModeInteractiveToolsContext* ToolsContext;
+
+	/** Command list lives here so that the key bindings on the commands can be processed in the viewport. */
+	TSharedPtr<FUICommandList> UICommandList;
+
+
+public:
+	/** Cached pointer to the viewport world interaction object we're using to interact with mesh elements */
+	class UViewportWorldInteraction* ViewportWorldInteraction;
+
+};
diff --git a/Engine/Plugins/Experimental/ModelingToolsEditorMode/Source/ModelingToolsEditorMode/Public/ModelingToolsEditorModeToolkit.h b/Engine/Plugins/Experimental/ModelingToolsEditorMode/Source/ModelingToolsEditorMode/Public/ModelingToolsEditorModeToolkit.h
new file mode 100644
--- /dev/null
+++ b/Engine/Plugins/Experimental/ModelingToolsEditorMode/Source/ModelingToolsEditorMode/Public/ModelingToolsEditorModeToolkit.h
@@ -0,0 +1,49 @@
+// Copyright 1998-2019 Epic Games, Inc. All Rights Reserved.
+
+#pragma once
+
+#include "CoreMinimal.h"
+#include "Toolkits/BaseToolkit.h"
+#include "InteractiveTool.h"
+#include "ModelingToolsEditorMode.h"
+
+#include "Widgets/SBoxPanel.h"
+
+class IDetailsView;
+class SButton;
+class STextBlock;
+
+class FModelingToolsEditorModeToolkit : public FModeToolkit
+{
+public:
+
+	FModelingToolsEditorModeToolkit();
+	
+	/** FModeToolkit interface */
+	virtual void Init(const TSharedPtr<IToolkitHost>& InitToolkitHost) override;
+
+	/** IToolkit interface */
+	virtual FName GetToolkitFName() const override;
+	virtual FText GetBaseToolkitName() const override;
+	virtual class FEdMode* GetEditorMode() const override;
+	virtual TSharedPtr<class SWidget> GetInlineContent() const override { return ToolkitWidget; }
+
+	virtual FModelingToolsEditorMode* GetToolsEditorMode() const;
+	virtual UEdModeInteractiveToolsContext* GetToolsContext() const;
+
+private:
+
+	TSharedPtr<SWidget> ToolkitWidget;
+	TSharedPtr<IDetailsView> DetailsView;
+
+	TSharedRef<SButton> MakeToolButton(const FText& ButtonLabel, const FString& ToolIdentifier);
+	SVerticalBox::FSlot& MakeToolButtonSlotV(const FText& ButtonLabel, const FString& ToolIdentifier);
+	SHorizontalBox::FSlot& MakeToolButtonSlotH(const FText& ButtonLabel, const FString& ToolIdentifier);
+
+	SVerticalBox::FSlot& MakeSetToolLabelV(const FText& LabelText);
+
+	TSharedPtr<STextBlock> ToolHeaderLabel;
+	TSharedPtr<SButton> AcceptButton;
+	TSharedPtr<SButton> CancelButton;
+	TSharedPtr<SButton> CompletedButton;
+};
