ModelingTools: use DetailsView customization to hide UObject labels in multi-object DetailsView
#rnx

diff --git a/Engine/Plugins/Experimental/ModelingToolsEditorMode/Source/ModelingToolsEditorMode/Private/ModelingToolsEditorModeToolkit.cpp b/Engine/Plugins/Experimental/ModelingToolsEditorMode/Source/ModelingToolsEditorMode/Private/ModelingToolsEditorModeToolkit.cpp
--- a/Engine/Plugins/Experimental/ModelingToolsEditorMode/Source/ModelingToolsEditorMode/Private/ModelingToolsEditorModeToolkit.cpp
+++ b/Engine/Plugins/Experimental/ModelingToolsEditorMode/Source/ModelingToolsEditorMode/Private/ModelingToolsEditorModeToolkit.cpp
@@ -7,6 +7,7 @@
 #include "Modules/ModuleManager.h"
 #include "PropertyEditorModule.h"
 #include "IDetailsView.h"
+#include "IDetailRootObjectCustomization.h"
 
 #include "Widgets/Input/SButton.h"
 #include "Widgets/Text/STextBlock.h"
@@ -21,6 +22,34 @@ FModelingToolsEditorModeToolkit::FModelingToolsEditorModeToolkit()
 
 
 
+/**
+ * Customization for Tool properties multi-object DetailsView, that just hides per-object header
+ */
+class FModelingToolsDetailRootObjectCustomization : public IDetailRootObjectCustomization
+{
+public:
+	FModelingToolsDetailRootObjectCustomization()
+	{
+	}
+
+	virtual TSharedPtr<SWidget> CustomizeObjectHeader(const UObject* InRootObject) override
+	{
+		return SNew(STextBlock).Text( FText::FromString(InRootObject->GetName())  );
+	}
+
+	virtual bool IsObjectVisible(const UObject* InRootObject) const override
+	{
+		return true;
+	}
+
+	virtual bool ShouldDisplayHeader(const UObject* InRootObject) const override
+	{
+		return false;
+	}
+};
+
+
+
 void FModelingToolsEditorModeToolkit::Init(const TSharedPtr<IToolkitHost>& InitToolkitHost)
 {
 
@@ -41,6 +70,11 @@ void FModelingToolsEditorModeToolkit::Init(const TSharedPtr<IToolkitHost>& InitT
 
 	DetailsView = PropertyEditorModule.CreateDetailView(DetailsViewArgs);
 
+	// add customization that hides header labels
+	TSharedPtr<FModelingToolsDetailRootObjectCustomization> RootObjectCustomization 
+		= MakeShared<FModelingToolsDetailRootObjectCustomization>();
+	DetailsView->SetRootObjectCustomizationInstance(RootObjectCustomization);
+
 
 	ToolHeaderLabel = SNew(STextBlock)
 		.AutoWrapText(true)
