Modeling Tools Plugin Updated to Use Modes Tool Palettes

diff --git a/Engine/Plugins/Experimental/ModelingToolsEditorMode/Source/ModelingToolsEditorMode/Private/ModelingToolsEditorMode.cpp b/Engine/Plugins/Experimental/ModelingToolsEditorMode/Source/ModelingToolsEditorMode/Private/ModelingToolsEditorMode.cpp
--- a/Engine/Plugins/Experimental/ModelingToolsEditorMode/Source/ModelingToolsEditorMode/Private/ModelingToolsEditorMode.cpp
+++ b/Engine/Plugins/Experimental/ModelingToolsEditorMode/Source/ModelingToolsEditorMode/Private/ModelingToolsEditorMode.cpp
@@ -55,15 +55,6 @@
 //#define ENABLE_DEBUG_PRINTING
 //#define TOOLED_ENABLE_VIEWPORT_INTERACTION
 
-
-// if set to 1, then on mode initialization we include buttons for prototype modeling tools
-static TAutoConsoleVariable<int32> CVarEnablePrototypeModelingTools(
-	TEXT("modeling.EnablePrototypes"),
-	0,
-	TEXT("Enable unsupported Experimental prototype Modeling Tools"));
-
-
-
 const FEditorModeID FModelingToolsEditorMode::EM_ModelingToolsEditorModeId = TEXT("EM_ModelingToolsEditorMode");
 
 FModelingToolsEditorMode::FModelingToolsEditorMode()
@@ -83,81 +74,17 @@ FModelingToolsEditorMode::~FModelingToolsEditorMode()
 }
 
 
-
-
-void FModelingToolsEditorMode::BuildModeToolbar(FToolBarBuilder& ToolbarBuilder)
-{
-	ToolbarBuilder.AddToolBarButton(FModelingToolsManagerCommands::Get().AcceptActiveTool);
-	ToolbarBuilder.AddToolBarButton(FModelingToolsManagerCommands::Get().CancelActiveTool);
-	ToolbarBuilder.AddToolBarButton(FModelingToolsManagerCommands::Get().CompleteActiveTool);
-
-	ToolbarBuilder.AddSeparator();
-
-	ToolbarBuilder.AddToolBarButton(FModelingToolsManagerCommands::Get().BeginAddPrimitiveTool);
-	ToolbarBuilder.AddToolBarButton(FModelingToolsManagerCommands::Get().BeginDrawPolygonTool);
-	if (CVarEnablePrototypeModelingTools.GetValueOnGameThread() > 0)
-	{
-		ToolbarBuilder.AddToolBarButton(FModelingToolsManagerCommands::Get().BeginShapeSprayTool);
-	}
-
-	ToolbarBuilder.AddSeparator();
-
-	ToolbarBuilder.AddToolBarButton(FModelingToolsManagerCommands::Get().BeginTransformMeshesTool);
-	ToolbarBuilder.AddToolBarButton(FModelingToolsManagerCommands::Get().BeginSculptMeshTool);
-	ToolbarBuilder.AddToolBarButton(FModelingToolsManagerCommands::Get().BeginPolyEditTool);
-	ToolbarBuilder.AddToolBarButton(FModelingToolsManagerCommands::Get().BeginSmoothMeshTool);
-	ToolbarBuilder.AddToolBarButton(FModelingToolsManagerCommands::Get().BeginDisplaceMeshTool);
-	if (CVarEnablePrototypeModelingTools.GetValueOnGameThread() > 0)
-	{
-		ToolbarBuilder.AddToolBarButton(FModelingToolsManagerCommands::Get().BeginMeshSpaceDeformerTool);
-	}
-
-	ToolbarBuilder.AddSeparator();
-
-	if (CVarEnablePrototypeModelingTools.GetValueOnGameThread() > 0)
-	{
-		ToolbarBuilder.AddToolBarButton(FModelingToolsManagerCommands::Get().BeginRemeshSculptMeshTool);
-	}
-	ToolbarBuilder.AddToolBarButton(FModelingToolsManagerCommands::Get().BeginMeshSelectionTool);
-	ToolbarBuilder.AddToolBarButton(FModelingToolsManagerCommands::Get().BeginPlaneCutTool);
-	ToolbarBuilder.AddToolBarButton(FModelingToolsManagerCommands::Get().BeginSimplifyMeshTool);
-	ToolbarBuilder.AddToolBarButton(FModelingToolsManagerCommands::Get().BeginRemeshMeshTool);
-	ToolbarBuilder.AddToolBarButton(FModelingToolsManagerCommands::Get().BeginUVProjectionTool);
-	ToolbarBuilder.AddToolBarButton(FModelingToolsManagerCommands::Get().BeginVoxelMergeTool);
-	ToolbarBuilder.AddToolBarButton(FModelingToolsManagerCommands::Get().BeginVoxelBooleanTool);
-	if (CVarEnablePrototypeModelingTools.GetValueOnGameThread() > 0)
-	{
-		ToolbarBuilder.AddToolBarButton(FModelingToolsManagerCommands::Get().BeginPolygonOnMeshTool);
-	}
-
-	ToolbarBuilder.AddSeparator();
-
-	ToolbarBuilder.AddToolBarButton(FModelingToolsManagerCommands::Get().BeginEditNormalsTool);
-	ToolbarBuilder.AddToolBarButton(FModelingToolsManagerCommands::Get().BeginWeldEdgesTool);
-	ToolbarBuilder.AddToolBarButton(FModelingToolsManagerCommands::Get().BeginMeshInspectorTool);
-	if (CVarEnablePrototypeModelingTools.GetValueOnGameThread() > 0)
-	{
-		ToolbarBuilder.AddToolBarButton(FModelingToolsManagerCommands::Get().BeginParameterizeMeshTool);
-		ToolbarBuilder.AddToolBarButton(FModelingToolsManagerCommands::Get().BeginPolyGroupsTool);
-	}
-
-}
-
-
-
 void FModelingToolsEditorMode::ActorSelectionChangeNotify()
 {
 }
 
 
-
 bool FModelingToolsEditorMode::ProcessEditDelete()
 {
 	return ToolsContext->ProcessEditDelete();
 }
 
 
-
 bool FModelingToolsEditorMode::AllowWidgetMove()
 { 
 	return false; 
diff --git a/Engine/Plugins/Experimental/ModelingToolsEditorMode/Source/ModelingToolsEditorMode/Private/ModelingToolsEditorModeToolkit.cpp b/Engine/Plugins/Experimental/ModelingToolsEditorMode/Source/ModelingToolsEditorMode/Private/ModelingToolsEditorModeToolkit.cpp
--- a/Engine/Plugins/Experimental/ModelingToolsEditorMode/Source/ModelingToolsEditorMode/Private/ModelingToolsEditorModeToolkit.cpp
+++ b/Engine/Plugins/Experimental/ModelingToolsEditorMode/Source/ModelingToolsEditorMode/Private/ModelingToolsEditorModeToolkit.cpp
@@ -2,8 +2,10 @@
 
 #include "ModelingToolsEditorModeToolkit.h"
 #include "ModelingToolsEditorMode.h"
+#include "ModelingToolsManagerActions.h"
 #include "Engine/Selection.h"
 
+#include "Framework/MultiBox/MultiBoxBuilder.h"
 #include "Modules/ModuleManager.h"
 #include "PropertyEditorModule.h"
 #include "IDetailsView.h"
@@ -16,6 +18,14 @@
 
 #define LOCTEXT_NAMESPACE "FModelingToolsEditorModeToolkit"
 
+
+// if set to 1, then on mode initialization we include buttons for prototype modeling tools
+static TAutoConsoleVariable<int32> CVarEnablePrototypeModelingTools(
+	TEXT("modeling.EnablePrototypes"),
+	1,
+	TEXT("Enable unsupported Experimental prototype Modeling Tools"));
+
+
 FModelingToolsEditorModeToolkit::FModelingToolsEditorModeToolkit()
 {
 }
@@ -182,4 +192,96 @@ UEdModeInteractiveToolsContext* FModelingToolsEditorModeToolkit::GetToolsContext
 	return GetToolsEditorMode()->GetToolsContext();
 }
 
+
+const TArray<FName> FModelingToolsEditorModeToolkit::PaletteNames = { FName(TEXT("Modeling")) /*, FName(TEXT("Experimental")) */ };
+
+FText FModelingToolsEditorModeToolkit::GetToolPaletteDisplayName(FName Palette) 
+{ 
+	return FText::FromName(Palette);
+}
+
+void FModelingToolsEditorModeToolkit::BuildToolPalette(FName PaletteIndex, class FToolBarBuilder& ToolbarBuilder) 
+{
+
+	const FModelingToolsManagerCommands& Commands = FModelingToolsManagerCommands::Get();
+
+	if (PaletteIndex == FName(TEXT("Modeling")))
+	{
+
+		ToolbarBuilder.AddToolBarButton(Commands.AcceptActiveTool);
+		ToolbarBuilder.AddToolBarButton(Commands.CancelActiveTool);
+		ToolbarBuilder.AddToolBarButton(Commands.CompleteActiveTool);
+
+		ToolbarBuilder.AddSeparator();
+
+		ToolbarBuilder.AddToolBarButton(Commands.BeginAddPrimitiveTool);
+		ToolbarBuilder.AddToolBarButton(Commands.BeginDrawPolygonTool);
+		if (CVarEnablePrototypeModelingTools.GetValueOnGameThread() > 0)
+		{
+			ToolbarBuilder.AddToolBarButton(Commands.BeginShapeSprayTool);
+		}
+
+		ToolbarBuilder.AddSeparator();
+
+		ToolbarBuilder.AddToolBarButton(Commands.BeginTransformMeshesTool);
+		ToolbarBuilder.AddToolBarButton(Commands.BeginSculptMeshTool);
+		ToolbarBuilder.AddToolBarButton(Commands.BeginPolyEditTool);
+		ToolbarBuilder.AddToolBarButton(Commands.BeginSmoothMeshTool);
+		ToolbarBuilder.AddToolBarButton(Commands.BeginDisplaceMeshTool);
+		if (CVarEnablePrototypeModelingTools.GetValueOnGameThread() > 0)
+		{
+			ToolbarBuilder.AddToolBarButton(Commands.BeginMeshSpaceDeformerTool);
+		}
+
+		ToolbarBuilder.AddSeparator();
+
+		if (CVarEnablePrototypeModelingTools.GetValueOnGameThread() > 0)
+		{
+			ToolbarBuilder.AddToolBarButton(Commands.BeginRemeshSculptMeshTool);
+		}
+		ToolbarBuilder.AddToolBarButton(Commands.BeginMeshSelectionTool);
+		ToolbarBuilder.AddToolBarButton(Commands.BeginPlaneCutTool);
+		ToolbarBuilder.AddToolBarButton(Commands.BeginSimplifyMeshTool);
+		ToolbarBuilder.AddToolBarButton(Commands.BeginRemeshMeshTool);
+		ToolbarBuilder.AddToolBarButton(Commands.BeginUVProjectionTool);
+		ToolbarBuilder.AddToolBarButton(Commands.BeginVoxelMergeTool);
+		ToolbarBuilder.AddToolBarButton(Commands.BeginVoxelBooleanTool);
+		if (CVarEnablePrototypeModelingTools.GetValueOnGameThread() > 0)
+		{
+			ToolbarBuilder.AddToolBarButton(Commands.BeginPolygonOnMeshTool);
+		}
+
+		ToolbarBuilder.AddSeparator();
+
+		ToolbarBuilder.AddToolBarButton(Commands.BeginEditNormalsTool);
+		ToolbarBuilder.AddToolBarButton(Commands.BeginWeldEdgesTool);
+		ToolbarBuilder.AddToolBarButton(Commands.BeginMeshInspectorTool);
+		if (CVarEnablePrototypeModelingTools.GetValueOnGameThread() > 0)
+		{
+			ToolbarBuilder.AddToolBarButton(Commands.BeginParameterizeMeshTool);
+			ToolbarBuilder.AddToolBarButton(Commands.BeginPolyGroupsTool);
+		}
+	}
+
+	else if (PaletteIndex == FName(TEXT("Experimental")))
+	{
+		ToolbarBuilder.AddToolBarButton(Commands.AcceptActiveTool);
+		ToolbarBuilder.AddToolBarButton(Commands.CancelActiveTool);
+		ToolbarBuilder.AddToolBarButton(Commands.CompleteActiveTool);
+
+		ToolbarBuilder.AddToolBarButton(Commands.BeginShapeSprayTool);
+		ToolbarBuilder.AddToolBarButton(Commands.BeginMeshSpaceDeformerTool);
+		ToolbarBuilder.AddToolBarButton(Commands.BeginRemeshSculptMeshTool);
+		ToolbarBuilder.AddToolBarButton(Commands.BeginPolygonOnMeshTool);
+		ToolbarBuilder.AddToolBarButton(Commands.BeginParameterizeMeshTool);
+		ToolbarBuilder.AddToolBarButton(Commands.BeginPolyGroupsTool);
+	}
+}
+
+
+void FModelingToolsEditorModeToolkit::OnToolPaletteChanged(FName PaletteName) 
+{
+}
+
+
 #undef LOCTEXT_NAMESPACE
diff --git a/Engine/Plugins/Experimental/ModelingToolsEditorMode/Source/ModelingToolsEditorMode/Public/ModelingToolsEditorMode.h b/Engine/Plugins/Experimental/ModelingToolsEditorMode/Source/ModelingToolsEditorMode/Public/ModelingToolsEditorMode.h
--- a/Engine/Plugins/Experimental/ModelingToolsEditorMode/Source/ModelingToolsEditorMode/Public/ModelingToolsEditorMode.h
+++ b/Engine/Plugins/Experimental/ModelingToolsEditorMode/Source/ModelingToolsEditorMode/Public/ModelingToolsEditorMode.h
@@ -42,8 +42,6 @@ public:
 
 	virtual bool ProcessEditDelete();
 
-	virtual void BuildModeToolbar(class FToolBarBuilder& ToolbarBuilder) override;
-
 
 	/*
 	 * focus events
diff --git a/Engine/Plugins/Experimental/ModelingToolsEditorMode/Source/ModelingToolsEditorMode/Public/ModelingToolsEditorModeToolkit.h b/Engine/Plugins/Experimental/ModelingToolsEditorMode/Source/ModelingToolsEditorMode/Public/ModelingToolsEditorModeToolkit.h
--- a/Engine/Plugins/Experimental/ModelingToolsEditorMode/Source/ModelingToolsEditorMode/Public/ModelingToolsEditorModeToolkit.h
+++ b/Engine/Plugins/Experimental/ModelingToolsEditorMode/Source/ModelingToolsEditorMode/Public/ModelingToolsEditorModeToolkit.h
@@ -35,6 +35,14 @@ public:
 	virtual void PostNotification(const FText& Message);
 	virtual void ClearNotification();
 
+	/** Returns the Mode specific tabs in the mode toolbar **/ 
+	const static TArray<FName> PaletteNames;
+	virtual void GetToolPaletteNames( TArray<FName>& InPaletteName ) const { InPaletteName = PaletteNames; }
+	virtual FText GetToolPaletteDisplayName(FName PaletteName); 
+	virtual void BuildToolPalette(FName PaletteName, class FToolBarBuilder& ToolbarBuilder);
+	virtual void OnToolPaletteChanged(FName PaletteName) override;
+
+
 private:
 
 	TSharedPtr<SWidget> ToolkitWidget;
