disable experimental modeling tools cvar by default
#rnx

diff --git a/Engine/Plugins/Experimental/ModelingToolsEditorMode/Source/ModelingToolsEditorMode/Private/ModelingToolsEditorModeToolkit.cpp b/Engine/Plugins/Experimental/ModelingToolsEditorMode/Source/ModelingToolsEditorMode/Private/ModelingToolsEditorModeToolkit.cpp
--- a/Engine/Plugins/Experimental/ModelingToolsEditorMode/Source/ModelingToolsEditorMode/Private/ModelingToolsEditorModeToolkit.cpp
+++ b/Engine/Plugins/Experimental/ModelingToolsEditorMode/Source/ModelingToolsEditorMode/Private/ModelingToolsEditorModeToolkit.cpp
@@ -22,7 +22,7 @@
 // if set to 1, then on mode initialization we include buttons for prototype modeling tools
 static TAutoConsoleVariable<int32> CVarEnablePrototypeModelingTools(
 	TEXT("modeling.EnablePrototypes"),
-	1,
+	0,
 	TEXT("Enable unsupported Experimental prototype Modeling Tools"));
 
 
