ModelingTools: reorganize tools into tabs when prototypes are enabled
#rnx

diff --git a/Engine/Plugins/Experimental/ModelingToolsEditorMode/Source/ModelingToolsEditorMode/Private/ModelingToolsEditorModeToolkit.cpp b/Engine/Plugins/Experimental/ModelingToolsEditorMode/Source/ModelingToolsEditorMode/Private/ModelingToolsEditorModeToolkit.cpp
--- a/Engine/Plugins/Experimental/ModelingToolsEditorMode/Source/ModelingToolsEditorMode/Private/ModelingToolsEditorModeToolkit.cpp
+++ b/Engine/Plugins/Experimental/ModelingToolsEditorMode/Source/ModelingToolsEditorMode/Private/ModelingToolsEditorModeToolkit.cpp
@@ -193,7 +193,16 @@ UEdModeInteractiveToolsContext* FModelingToolsEditorModeToolkit::GetToolsContext
 }
 
 
-const TArray<FName> FModelingToolsEditorModeToolkit::PaletteNames = { FName(TEXT("Modeling")) /*, FName(TEXT("Experimental")) */ };
+const TArray<FName> FModelingToolsEditorModeToolkit::PaletteNames_Standard = { FName(TEXT("Modeling")) };
+const TArray<FName> FModelingToolsEditorModeToolkit::PaletteNames_Experimental = { FName(TEXT("Create")), FName(TEXT("Edit")), FName(TEXT("UVs/Normals")), FName(TEXT("Utilities")) };
+
+
+void FModelingToolsEditorModeToolkit::GetToolPaletteNames(TArray<FName>& InPaletteName) const
+{
+	bool bEnablePrototypes = (CVarEnablePrototypeModelingTools.GetValueOnGameThread() > 0);
+	InPaletteName = (bEnablePrototypes) ? PaletteNames_Experimental : PaletteNames_Standard;
+}
+
 
 FText FModelingToolsEditorModeToolkit::GetToolPaletteDisplayName(FName Palette) 
 { 
@@ -202,12 +211,25 @@ FText FModelingToolsEditorModeToolkit::GetToolPaletteDisplayName(FName Palette)
 
 void FModelingToolsEditorModeToolkit::BuildToolPalette(FName PaletteIndex, class FToolBarBuilder& ToolbarBuilder) 
 {
-
 	const FModelingToolsManagerCommands& Commands = FModelingToolsManagerCommands::Get();
+	bool bEnablePrototypes = (CVarEnablePrototypeModelingTools.GetValueOnGameThread() > 0);
+	if (bEnablePrototypes)
+	{
+		BuildToolPalette_Experimental(PaletteIndex, ToolbarBuilder);
+	}
+	else
+	{
+		BuildToolPalette_Standard(PaletteIndex, ToolbarBuilder);
+	}
+}
+
 
+
+void FModelingToolsEditorModeToolkit::BuildToolPalette_Standard(FName PaletteIndex, class FToolBarBuilder& ToolbarBuilder)
+{
+	const FModelingToolsManagerCommands& Commands = FModelingToolsManagerCommands::Get();
 	if (PaletteIndex == FName(TEXT("Modeling")))
 	{
-
 		ToolbarBuilder.AddToolBarButton(Commands.AcceptActiveTool);
 		ToolbarBuilder.AddToolBarButton(Commands.CancelActiveTool);
 		ToolbarBuilder.AddToolBarButton(Commands.CompleteActiveTool);
@@ -216,10 +238,6 @@ void FModelingToolsEditorModeToolkit::BuildToolPalette(FName PaletteIndex, class
 
 		ToolbarBuilder.AddToolBarButton(Commands.BeginAddPrimitiveTool);
 		ToolbarBuilder.AddToolBarButton(Commands.BeginDrawPolygonTool);
-		if (CVarEnablePrototypeModelingTools.GetValueOnGameThread() > 0)
-		{
-			ToolbarBuilder.AddToolBarButton(Commands.BeginShapeSprayTool);
-		}
 
 		ToolbarBuilder.AddSeparator();
 
@@ -228,17 +246,9 @@ void FModelingToolsEditorModeToolkit::BuildToolPalette(FName PaletteIndex, class
 		ToolbarBuilder.AddToolBarButton(Commands.BeginPolyEditTool);
 		ToolbarBuilder.AddToolBarButton(Commands.BeginSmoothMeshTool);
 		ToolbarBuilder.AddToolBarButton(Commands.BeginDisplaceMeshTool);
-		if (CVarEnablePrototypeModelingTools.GetValueOnGameThread() > 0)
-		{
-			ToolbarBuilder.AddToolBarButton(Commands.BeginMeshSpaceDeformerTool);
-		}
 
 		ToolbarBuilder.AddSeparator();
 
-		if (CVarEnablePrototypeModelingTools.GetValueOnGameThread() > 0)
-		{
-			ToolbarBuilder.AddToolBarButton(Commands.BeginRemeshSculptMeshTool);
-		}
 		ToolbarBuilder.AddToolBarButton(Commands.BeginMeshSelectionTool);
 		ToolbarBuilder.AddToolBarButton(Commands.BeginPlaneCutTool);
 		ToolbarBuilder.AddToolBarButton(Commands.BeginSimplifyMeshTool);
@@ -246,39 +256,73 @@ void FModelingToolsEditorModeToolkit::BuildToolPalette(FName PaletteIndex, class
 		ToolbarBuilder.AddToolBarButton(Commands.BeginUVProjectionTool);
 		ToolbarBuilder.AddToolBarButton(Commands.BeginVoxelMergeTool);
 		ToolbarBuilder.AddToolBarButton(Commands.BeginVoxelBooleanTool);
-		if (CVarEnablePrototypeModelingTools.GetValueOnGameThread() > 0)
-		{
-			ToolbarBuilder.AddToolBarButton(Commands.BeginPolygonOnMeshTool);
-		}
 
 		ToolbarBuilder.AddSeparator();
 
 		ToolbarBuilder.AddToolBarButton(Commands.BeginEditNormalsTool);
 		ToolbarBuilder.AddToolBarButton(Commands.BeginWeldEdgesTool);
 		ToolbarBuilder.AddToolBarButton(Commands.BeginMeshInspectorTool);
-		if (CVarEnablePrototypeModelingTools.GetValueOnGameThread() > 0)
-		{
-			ToolbarBuilder.AddToolBarButton(Commands.BeginParameterizeMeshTool);
-			ToolbarBuilder.AddToolBarButton(Commands.BeginPolyGroupsTool);
-		}
 	}
+}
 
-	else if (PaletteIndex == FName(TEXT("Experimental")))
-	{
-		ToolbarBuilder.AddToolBarButton(Commands.AcceptActiveTool);
-		ToolbarBuilder.AddToolBarButton(Commands.CancelActiveTool);
-		ToolbarBuilder.AddToolBarButton(Commands.CompleteActiveTool);
+void FModelingToolsEditorModeToolkit::BuildToolPalette_Experimental(FName PaletteIndex, class FToolBarBuilder& ToolbarBuilder)
+{
+	const FModelingToolsManagerCommands& Commands = FModelingToolsManagerCommands::Get();
+
+	ToolbarBuilder.AddToolBarButton(Commands.AcceptActiveTool);
+	ToolbarBuilder.AddToolBarButton(Commands.CancelActiveTool);
+	ToolbarBuilder.AddToolBarButton(Commands.CompleteActiveTool);
 
+	ToolbarBuilder.AddSeparator();
+
+
+	if (PaletteIndex == FName(TEXT("Create")) )
+	{
+		ToolbarBuilder.AddToolBarButton(Commands.BeginAddPrimitiveTool);
+		ToolbarBuilder.AddToolBarButton(Commands.BeginDrawPolygonTool);
 		ToolbarBuilder.AddToolBarButton(Commands.BeginShapeSprayTool);
-		ToolbarBuilder.AddToolBarButton(Commands.BeginMeshSpaceDeformerTool);
+		ToolbarBuilder.AddToolBarButton(Commands.BeginVoxelMergeTool);
+		ToolbarBuilder.AddToolBarButton(Commands.BeginVoxelBooleanTool);
+	} 
+	else if ( PaletteIndex == FName(TEXT("Edit")) ) 
+	{
+		ToolbarBuilder.AddToolBarButton(Commands.BeginTransformMeshesTool);
+		ToolbarBuilder.AddToolBarButton(Commands.BeginMeshSelectionTool);
+
+		ToolbarBuilder.AddSeparator();
+
+		ToolbarBuilder.AddToolBarButton(Commands.BeginSculptMeshTool);
 		ToolbarBuilder.AddToolBarButton(Commands.BeginRemeshSculptMeshTool);
+
+		ToolbarBuilder.AddSeparator();
+
+		ToolbarBuilder.AddToolBarButton(Commands.BeginPolyEditTool);
+		ToolbarBuilder.AddToolBarButton(Commands.BeginSmoothMeshTool);
+		ToolbarBuilder.AddToolBarButton(Commands.BeginDisplaceMeshTool);
+		ToolbarBuilder.AddToolBarButton(Commands.BeginMeshSpaceDeformerTool);
+
+		ToolbarBuilder.AddSeparator();
+
+		ToolbarBuilder.AddToolBarButton(Commands.BeginPlaneCutTool);
 		ToolbarBuilder.AddToolBarButton(Commands.BeginPolygonOnMeshTool);
+
+	}
+	else if (PaletteIndex == FName(TEXT("UVs/Normals")))
+	{
+		ToolbarBuilder.AddToolBarButton(Commands.BeginUVProjectionTool);
 		ToolbarBuilder.AddToolBarButton(Commands.BeginParameterizeMeshTool);
+		ToolbarBuilder.AddToolBarButton(Commands.BeginEditNormalsTool);
+	}
+	else if (PaletteIndex == FName(TEXT("Utilities")))
+	{
 		ToolbarBuilder.AddToolBarButton(Commands.BeginPolyGroupsTool);
+		ToolbarBuilder.AddToolBarButton(Commands.BeginSimplifyMeshTool);
+		ToolbarBuilder.AddToolBarButton(Commands.BeginRemeshMeshTool);
+		ToolbarBuilder.AddToolBarButton(Commands.BeginWeldEdgesTool);
+		ToolbarBuilder.AddToolBarButton(Commands.BeginMeshInspectorTool);
 	}
 }
 
-
 void FModelingToolsEditorModeToolkit::OnToolPaletteChanged(FName PaletteName) 
 {
 }
diff --git a/Engine/Plugins/Experimental/ModelingToolsEditorMode/Source/ModelingToolsEditorMode/Public/ModelingToolsEditorModeToolkit.h b/Engine/Plugins/Experimental/ModelingToolsEditorMode/Source/ModelingToolsEditorMode/Public/ModelingToolsEditorModeToolkit.h
--- a/Engine/Plugins/Experimental/ModelingToolsEditorMode/Source/ModelingToolsEditorMode/Public/ModelingToolsEditorModeToolkit.h
+++ b/Engine/Plugins/Experimental/ModelingToolsEditorMode/Source/ModelingToolsEditorMode/Public/ModelingToolsEditorModeToolkit.h
@@ -36,14 +36,18 @@ public:
 	virtual void ClearNotification();
 
 	/** Returns the Mode specific tabs in the mode toolbar **/ 
-	const static TArray<FName> PaletteNames;
-	virtual void GetToolPaletteNames( TArray<FName>& InPaletteName ) const { InPaletteName = PaletteNames; }
+	virtual void GetToolPaletteNames(TArray<FName>& InPaletteName) const;
 	virtual FText GetToolPaletteDisplayName(FName PaletteName); 
 	virtual void BuildToolPalette(FName PaletteName, class FToolBarBuilder& ToolbarBuilder);
 	virtual void OnToolPaletteChanged(FName PaletteName) override;
 
+	virtual void BuildToolPalette_Standard(FName PaletteName, class FToolBarBuilder& ToolbarBuilder);
+	virtual void BuildToolPalette_Experimental(FName PaletteName, class FToolBarBuilder& ToolbarBuilder);
 
 private:
+	const static TArray<FName> PaletteNames_Standard;
+	const static TArray<FName> PaletteNames_Experimental;
+
 
 	TSharedPtr<SWidget> ToolkitWidget;
 	TSharedPtr<IDetailsView> DetailsView;
