ModelingTools: tools must each have their own TCommands if they are to re-use hotkeys. Separate out Sculpt and Transform tools, and add utility API functions to top-level FModelingToolActionCommands which hide these per-tool TCommands.
#rnx

diff --git a/Engine/Plugins/Experimental/ModelingToolsEditorMode/Source/ModelingToolsEditorMode/Private/ModelingToolsActions.cpp b/Engine/Plugins/Experimental/ModelingToolsEditorMode/Source/ModelingToolsEditorMode/Private/ModelingToolsActions.cpp
--- a/Engine/Plugins/Experimental/ModelingToolsEditorMode/Source/ModelingToolsEditorMode/Private/ModelingToolsActions.cpp
+++ b/Engine/Plugins/Experimental/ModelingToolsEditorMode/Source/ModelingToolsEditorMode/Private/ModelingToolsActions.cpp
@@ -29,13 +29,84 @@ FModelingToolActionCommands::FModelingToolActionCommands() :
 
 void FModelingToolActionCommands::GetToolDefaultObjectList(TArray<UInteractiveTool*>& ToolCDOs)
 {
-	ToolCDOs.Add(GetMutableDefault<UDynamicMeshSculptTool>());
 	ToolCDOs.Add(GetMutableDefault<UMeshInspectorTool>());
 	ToolCDOs.Add(GetMutableDefault<UDrawPolygonTool>());
 	ToolCDOs.Add(GetMutableDefault<UEditMeshPolygonsTool>());
 	ToolCDOs.Add(GetMutableDefault<UMeshSpaceDeformerTool>());
 	ToolCDOs.Add(GetMutableDefault<UShapeSprayTool>());
 	ToolCDOs.Add(GetMutableDefault<UMeshSelectionTool>());
+}
+
+
+
+void FModelingToolActionCommands::RegisterAllToolActions()
+{
+	FModelingToolActionCommands::Register();
+	FSculptToolActionCommands::Register();
+	FTransformToolActionCommands::Register();
+}
+
+void FModelingToolActionCommands::UnregisterAllToolActions()
+{
+	FModelingToolActionCommands::Unregister();
+	FSculptToolActionCommands::Unregister();
+	FTransformToolActionCommands::Unregister();
+}
+
+
+
+void FModelingToolActionCommands::UpdateToolCommandBinding(UInteractiveTool* Tool, TSharedPtr<FUICommandList> UICommandList, bool bUnbind)
+{
+#define UPDATE_BINDING(CommandsType)  if (!bUnbind) CommandsType::Get().BindCommandsForCurrentTool(UICommandList, Tool); else CommandsType::Get().UnbindActiveCommands(UICommandList);
+
+
+	if (Cast<UTransformMeshesTool>(Tool) != nullptr)
+	{
+		UPDATE_BINDING(FTransformToolActionCommands);
+	}
+	else if (Cast<UDynamicMeshSculptTool>(Tool) != nullptr)
+	{
+		UPDATE_BINDING(FSculptToolActionCommands);
+	}
+	else
+	{
+		UPDATE_BINDING(FModelingToolActionCommands);
+	}
+}
+
+
+
+
+
+FSculptToolActionCommands::FSculptToolActionCommands() :
+	TInteractiveToolCommands<FSculptToolActionCommands>(
+		"ModelingToolsSculptTool", // Context name for fast lookup
+		NSLOCTEXT("Contexts", "ModelingToolsSculptTool", "Modeling Tools - Sculpt Tool"), // Localized context name for displaying
+		NAME_None, // Parent
+		FEditorStyle::GetStyleSetName() // Icon Style Set
+		)
+{
+}
+
+void FSculptToolActionCommands::GetToolDefaultObjectList(TArray<UInteractiveTool*>& ToolCDOs)
+{
+	ToolCDOs.Add(GetMutableDefault<UDynamicMeshSculptTool>());
+}
+
+
+
+FTransformToolActionCommands::FTransformToolActionCommands() :
+	TInteractiveToolCommands<FTransformToolActionCommands>(
+		"ModelingToolsTransformTool", // Context name for fast lookup
+		NSLOCTEXT("Contexts", "ModelingToolsTransformTool", "Modeling Tools - Transform Tool"), // Localized context name for displaying
+		NAME_None, // Parent
+		FEditorStyle::GetStyleSetName() // Icon Style Set
+		)
+{
+}
+
+void FTransformToolActionCommands::GetToolDefaultObjectList(TArray<UInteractiveTool*>& ToolCDOs)
+{
 	ToolCDOs.Add(GetMutableDefault<UTransformMeshesTool>());
 }
 
diff --git a/Engine/Plugins/Experimental/ModelingToolsEditorMode/Source/ModelingToolsEditorMode/Private/ModelingToolsEditorMode.cpp b/Engine/Plugins/Experimental/ModelingToolsEditorMode/Source/ModelingToolsEditorMode/Private/ModelingToolsEditorMode.cpp
--- a/Engine/Plugins/Experimental/ModelingToolsEditorMode/Source/ModelingToolsEditorMode/Private/ModelingToolsEditorMode.cpp
+++ b/Engine/Plugins/Experimental/ModelingToolsEditorMode/Source/ModelingToolsEditorMode/Private/ModelingToolsEditorMode.cpp
@@ -393,13 +393,11 @@ void FModelingToolsEditorMode::Enter()
 	// listen for Tool start/end events to bind/unbind any hotkeys relevant to that Tool
 	ToolsContext->ToolManager->OnToolStarted.AddLambda([this](UInteractiveToolManager* Manager, UInteractiveTool* Tool)
 	{
-		const FModelingToolActionCommands& Commands = FModelingToolActionCommands::Get();
-		Commands.BindCommandsForCurrentTool(UICommandList, Tool);
+		FModelingToolActionCommands::UpdateToolCommandBinding(Tool, UICommandList, false);
 	});
 	ToolsContext->ToolManager->OnToolEnded.AddLambda([this](UInteractiveToolManager* Manager, UInteractiveTool* Tool)
 	{
-		const FModelingToolActionCommands& Commands = FModelingToolActionCommands::Get();
-		Commands.UnbindActiveCommands(UICommandList);
+		FModelingToolActionCommands::UpdateToolCommandBinding(Tool, UICommandList, true);
 	});
 
 
diff --git a/Engine/Plugins/Experimental/ModelingToolsEditorMode/Source/ModelingToolsEditorMode/Private/ModelingToolsEditorModeModule.cpp b/Engine/Plugins/Experimental/ModelingToolsEditorMode/Source/ModelingToolsEditorMode/Private/ModelingToolsEditorModeModule.cpp
--- a/Engine/Plugins/Experimental/ModelingToolsEditorMode/Source/ModelingToolsEditorMode/Private/ModelingToolsEditorModeModule.cpp
+++ b/Engine/Plugins/Experimental/ModelingToolsEditorMode/Source/ModelingToolsEditorMode/Private/ModelingToolsEditorModeModule.cpp
@@ -18,7 +18,7 @@ void FModelingToolsEditorModeModule::ShutdownModule()
 {
 	FCoreDelegates::OnPostEngineInit.RemoveAll(this);
 
-	FModelingToolActionCommands::Unregister();
+	FModelingToolActionCommands::UnregisterAllToolActions();
 	FModelingToolsManagerCommands::Unregister();
 
 	// Unregister slate style overrides
@@ -42,7 +42,7 @@ void FModelingToolsEditorModeModule::OnPostEngineInit()
 		FSlateIcon(),
 		true);
 
-	FModelingToolActionCommands::Register();
+	FModelingToolActionCommands::RegisterAllToolActions();
 	FModelingToolsManagerCommands::Register();
 }
 
diff --git a/Engine/Plugins/Experimental/ModelingToolsEditorMode/Source/ModelingToolsEditorMode/Public/ModelingToolsActions.h b/Engine/Plugins/Experimental/ModelingToolsEditorMode/Source/ModelingToolsEditorMode/Public/ModelingToolsActions.h
--- a/Engine/Plugins/Experimental/ModelingToolsEditorMode/Source/ModelingToolsEditorMode/Public/ModelingToolsActions.h
+++ b/Engine/Plugins/Experimental/ModelingToolsEditorMode/Source/ModelingToolsEditorMode/Public/ModelingToolsActions.h
@@ -6,6 +6,7 @@
 #include "Tools/InteractiveToolsCommands.h"
 
 
+
 /**
  * TInteractiveToolCommands implementation for this module that provides standard Editor hotkey support
  */
@@ -15,5 +16,47 @@ public:
 	FModelingToolActionCommands();
 
 	virtual void GetToolDefaultObjectList(TArray<UInteractiveTool*>& ToolCDOs) override;
+
+	/**
+	 * interface that hides various per-tool action sets
+	 */
+
+	/**
+	 * Register all Tool command sets. Call this in module startup
+	 */
+	static void RegisterAllToolActions();
+
+	/**
+	 * Unregister all Tool command sets. Call this from module shutdown.
+	 */
+	static void UnregisterAllToolActions();
+
+	/**
+	 * Add or remove commands relevant to Tool to the given UICommandList.
+	 * Call this when the active tool changes (eg on ToolManager.OnToolStarted / OnToolEnded)
+	 * @param bUnbind if true, commands are removed, otherwise added
+	 */
+	static void UpdateToolCommandBinding(UInteractiveTool* Tool, TSharedPtr<FUICommandList> UICommandList, bool bUnbind = false);
+
 };
 
+
+// Each TCommands can only have a single (chord,action) binding, regardless of whether these
+// would ever be used at the same time. And we cannot create/register TCommands at runtime.
+// So, we have to define a separate TCommands instance for each Tool. This is unfortunate.
+
+
+class FSculptToolActionCommands : public TInteractiveToolCommands<FSculptToolActionCommands>
+{
+public:
+	FSculptToolActionCommands();
+	virtual void GetToolDefaultObjectList(TArray<UInteractiveTool*>& ToolCDOs) override;
+};
+
+
+class FTransformToolActionCommands : public TInteractiveToolCommands<FTransformToolActionCommands>
+{
+public:
+	FTransformToolActionCommands();
+	virtual void GetToolDefaultObjectList(TArray<UInteractiveTool*>& ToolCDOs) override;
+};
