MeshModelingTools: Details View now reacts to resizing and expands to fill available height

diff --git a/Engine/Plugins/Experimental/ModelingToolsEditorMode/Source/ModelingToolsEditorMode/Private/ModelingToolsEditorModeToolkit.cpp b/Engine/Plugins/Experimental/ModelingToolsEditorMode/Source/ModelingToolsEditorMode/Private/ModelingToolsEditorModeToolkit.cpp
--- a/Engine/Plugins/Experimental/ModelingToolsEditorMode/Source/ModelingToolsEditorMode/Private/ModelingToolsEditorModeToolkit.cpp
+++ b/Engine/Plugins/Experimental/ModelingToolsEditorMode/Source/ModelingToolsEditorMode/Private/ModelingToolsEditorModeToolkit.cpp
@@ -85,7 +85,6 @@ void FModelingToolsEditorModeToolkit::Init(const TSharedPtr<IToolkitHost>& InitT
 		= MakeShared<FModelingToolsDetailRootObjectCustomization>();
 	DetailsView->SetRootObjectCustomizationInstance(RootObjectCustomization);
 
-
 	ToolHeaderLabel = SNew(STextBlock)
 		.AutoWrapText(true)
 		.Font(FCoreStyle::GetDefaultFontStyle("Bold", 12));
@@ -117,7 +116,7 @@ void FModelingToolsEditorModeToolkit::Init(const TSharedPtr<IToolkitHost>& InitT
 					ToolMessageArea->AsShared()
 				]
 
-			+ SVerticalBox::Slot().HAlign(HAlign_Fill).AutoHeight().MaxHeight(500.0f)
+			+ SVerticalBox::Slot().HAlign(HAlign_Fill).FillHeight(1.f)
 				[
 					DetailsView->AsShared()
 				]
