ModelingMode: Add warning message type and warning message area to Modeling Tools Mode.
#rnx

diff --git a/Engine/Plugins/Experimental/ModelingToolsEditorMode/Source/ModelingToolsEditorMode/Private/ModelingToolsEditorMode.cpp b/Engine/Plugins/Experimental/ModelingToolsEditorMode/Source/ModelingToolsEditorMode/Private/ModelingToolsEditorMode.cpp
--- a/Engine/Plugins/Experimental/ModelingToolsEditorMode/Source/ModelingToolsEditorMode/Private/ModelingToolsEditorMode.cpp
+++ b/Engine/Plugins/Experimental/ModelingToolsEditorMode/Source/ModelingToolsEditorMode/Private/ModelingToolsEditorMode.cpp
@@ -253,7 +253,10 @@ void FModelingToolsEditorMode::Enter()
 	{
 		this->OnToolNotificationMessage.Broadcast(Message);
 	});
-
+	ToolsContext->OnToolWarningMessage.AddLambda([this](const FText& Message)
+	{
+		this->OnToolWarningMessage.Broadcast(Message);
+	});
 
 	if (!Toolkit.IsValid() && UsesToolkits())
 	{
diff --git a/Engine/Plugins/Experimental/ModelingToolsEditorMode/Source/ModelingToolsEditorMode/Private/ModelingToolsEditorModeToolkit.cpp b/Engine/Plugins/Experimental/ModelingToolsEditorMode/Source/ModelingToolsEditorMode/Private/ModelingToolsEditorModeToolkit.cpp
--- a/Engine/Plugins/Experimental/ModelingToolsEditorMode/Source/ModelingToolsEditorMode/Private/ModelingToolsEditorModeToolkit.cpp
+++ b/Engine/Plugins/Experimental/ModelingToolsEditorMode/Source/ModelingToolsEditorMode/Private/ModelingToolsEditorModeToolkit.cpp
@@ -100,6 +100,15 @@ void FModelingToolsEditorModeToolkit::Init(const TSharedPtr<IToolkitHost>& InitT
 		.ColorAndOpacity(FSlateColor(FLinearColor::White * 0.7f));
 	ToolMessageArea->SetText(LOCTEXT("ToolMessageLabel", ""));
 
+	ToolWarningArea = SNew(STextBlock)
+		.AutoWrapText(true)
+		.Font(FCoreStyle::GetDefaultFontStyle("Bold", 9))
+		.ColorAndOpacity(FSlateColor(FLinearColor(0.9f, 0.15f, 0.15f)));
+	ToolWarningArea->SetText(LOCTEXT("ToolMessageLabel", ""));
+
+
+	
+
 	SAssignNew(ToolkitWidget, SBorder)
 		.HAlign(HAlign_Fill)
 		.Padding(4)
@@ -111,11 +120,16 @@ void FModelingToolsEditorModeToolkit::Init(const TSharedPtr<IToolkitHost>& InitT
 					ToolHeaderLabel->AsShared()
 				]
 
-			+ SVerticalBox::Slot().AutoHeight().HAlign(HAlign_Fill).AutoHeight().Padding(10, 10, 10, 20)
+			+ SVerticalBox::Slot().AutoHeight().HAlign(HAlign_Fill).AutoHeight().Padding(10, 10, 10, 15)
 				[
 					ToolMessageArea->AsShared()
 				]
 
+			+ SVerticalBox::Slot().AutoHeight().HAlign(HAlign_Fill).AutoHeight().Padding(10, 5, 10, 15)
+				[
+					ToolWarningArea->AsShared()
+				]
+
 			+ SVerticalBox::Slot().HAlign(HAlign_Fill).FillHeight(1.f)
 				[
 					DetailsView->AsShared()
@@ -125,6 +139,8 @@ void FModelingToolsEditorModeToolkit::Init(const TSharedPtr<IToolkitHost>& InitT
 		
 	FModeToolkit::Init(InitToolkitHost);
 
+	ClearNotification();
+	ClearWarning();
 
 
 	GetToolsEditorMode()->GetToolManager()->OnToolStarted.AddLambda([this](UInteractiveToolManager* Manager, UInteractiveTool* Tool)
@@ -141,6 +157,7 @@ void FModelingToolsEditorModeToolkit::Init(const TSharedPtr<IToolkitHost>& InitT
 		DetailsView->SetObject(nullptr);
 		ToolHeaderLabel->SetText(LOCTEXT("SelectToolLabel", "Select a Tool from the Toolbar"));
 		ClearNotification();
+		ClearWarning();
 	});
 
 
@@ -148,8 +165,10 @@ void FModelingToolsEditorModeToolkit::Init(const TSharedPtr<IToolkitHost>& InitT
 	{
 		PostNotification(Message);
 	});
-
-
+	GetToolsEditorMode()->OnToolWarningMessage.AddLambda([this](const FText& Message)
+	{
+		PostWarning(Message);
+	});
 }
 
 
@@ -157,11 +176,26 @@ void FModelingToolsEditorModeToolkit::Init(const TSharedPtr<IToolkitHost>& InitT
 void FModelingToolsEditorModeToolkit::PostNotification(const FText& Message)
 {
 	ToolMessageArea->SetText(Message);
+	ToolMessageArea->SetVisibility(EVisibility::Visible);
 }
 
 void FModelingToolsEditorModeToolkit::ClearNotification()
 {
 	ToolMessageArea->SetText(FText());
+	ToolMessageArea->SetVisibility(EVisibility::Collapsed);
+}
+
+
+void FModelingToolsEditorModeToolkit::PostWarning(const FText& Message)
+{
+	ToolWarningArea->SetText(Message);
+	ToolWarningArea->SetVisibility(EVisibility::Visible);
+}
+
+void FModelingToolsEditorModeToolkit::ClearWarning()
+{
+	ToolWarningArea->SetText(FText());
+	ToolWarningArea->SetVisibility(EVisibility::Collapsed);
 }
 
 
@@ -241,13 +275,17 @@ void FModelingToolsEditorModeToolkit::BuildToolPalette_Standard(FName PaletteInd
 		ToolbarBuilder.AddSeparator();
 
 		ToolbarBuilder.AddToolBarButton(Commands.BeginTransformMeshesTool);
-		ToolbarBuilder.AddToolBarButton(Commands.BeginSculptMeshTool);
 		ToolbarBuilder.AddToolBarButton(Commands.BeginPolyEditTool);
 		ToolbarBuilder.AddToolBarButton(Commands.BeginSmoothMeshTool);
 		ToolbarBuilder.AddToolBarButton(Commands.BeginDisplaceMeshTool);
 
 		ToolbarBuilder.AddSeparator();
 
+		ToolbarBuilder.AddToolBarButton(Commands.BeginSculptMeshTool);
+		ToolbarBuilder.AddToolBarButton(Commands.BeginRemeshSculptMeshTool);
+
+		ToolbarBuilder.AddSeparator();
+
 		ToolbarBuilder.AddToolBarButton(Commands.BeginMeshSelectionTool);
 		ToolbarBuilder.AddToolBarButton(Commands.BeginPlaneCutTool);
 		ToolbarBuilder.AddToolBarButton(Commands.BeginSimplifyMeshTool);
diff --git a/Engine/Plugins/Experimental/ModelingToolsEditorMode/Source/ModelingToolsEditorMode/Public/ModelingToolsEditorMode.h b/Engine/Plugins/Experimental/ModelingToolsEditorMode/Source/ModelingToolsEditorMode/Public/ModelingToolsEditorMode.h
--- a/Engine/Plugins/Experimental/ModelingToolsEditorMode/Source/ModelingToolsEditorMode/Public/ModelingToolsEditorMode.h
+++ b/Engine/Plugins/Experimental/ModelingToolsEditorMode/Source/ModelingToolsEditorMode/Public/ModelingToolsEditorMode.h
@@ -135,6 +135,7 @@ public:
 public:
 	DECLARE_MULTICAST_DELEGATE_OneParam(FOnModelingModeToolNotification, const FText&);
 	FOnModelingModeToolNotification OnToolNotificationMessage;
+	FOnModelingModeToolNotification OnToolWarningMessage;
 
 
 public:
diff --git a/Engine/Plugins/Experimental/ModelingToolsEditorMode/Source/ModelingToolsEditorMode/Public/ModelingToolsEditorModeToolkit.h b/Engine/Plugins/Experimental/ModelingToolsEditorMode/Source/ModelingToolsEditorMode/Public/ModelingToolsEditorModeToolkit.h
--- a/Engine/Plugins/Experimental/ModelingToolsEditorMode/Source/ModelingToolsEditorMode/Public/ModelingToolsEditorModeToolkit.h
+++ b/Engine/Plugins/Experimental/ModelingToolsEditorMode/Source/ModelingToolsEditorMode/Public/ModelingToolsEditorModeToolkit.h
@@ -35,6 +35,10 @@ public:
 	virtual void PostNotification(const FText& Message);
 	virtual void ClearNotification();
 
+	// set/clear warning message area
+	virtual void PostWarning(const FText& Message);
+	virtual void ClearWarning();
+
 	/** Returns the Mode specific tabs in the mode toolbar **/ 
 	virtual void GetToolPaletteNames(TArray<FName>& InPaletteName) const;
 	virtual FText GetToolPaletteDisplayName(FName PaletteName); 
@@ -54,6 +58,7 @@ private:
 
 	TSharedPtr<STextBlock> ToolHeaderLabel;
 	TSharedPtr<STextBlock> ToolMessageArea;
+	TSharedPtr<STextBlock> ToolWarningArea;
 	TSharedPtr<SButton> AcceptButton;
 	TSharedPtr<SButton> CancelButton;
 	TSharedPtr<SButton> CompletedButton;
diff --git a/Engine/Source/Editor/Experimental/EditorInteractiveToolsFramework/Private/EdModeInteractiveToolsContext.cpp b/Engine/Source/Editor/Experimental/EditorInteractiveToolsFramework/Private/EdModeInteractiveToolsContext.cpp
--- a/Engine/Source/Editor/Experimental/EditorInteractiveToolsFramework/Private/EdModeInteractiveToolsContext.cpp
+++ b/Engine/Source/Editor/Experimental/EditorInteractiveToolsFramework/Private/EdModeInteractiveToolsContext.cpp
@@ -253,6 +253,10 @@ public:
 		{
 			ToolsContext->PostToolNotificationMessage(Message);
 		}
+		if (Level == EToolMessageLevel::UserWarning)
+		{
+			ToolsContext->PostToolWarningMessage(Message);
+		}
 		else
 		{
 			UE_LOG(LogTemp, Warning, TEXT("%s"), *Message.ToString());
@@ -415,6 +419,11 @@ void UEdModeInteractiveToolsContext::PostToolNotificationMessage(const FText& Me
 	OnToolNotificationMessage.Broadcast(Message);
 }
 
+void UEdModeInteractiveToolsContext::PostToolWarningMessage(const FText& Message)
+{
+	OnToolWarningMessage.Broadcast(Message);
+}
+
 
 
 void UEdModeInteractiveToolsContext::TerminateActiveToolsOnPIEStart()
diff --git a/Engine/Source/Editor/Experimental/EditorInteractiveToolsFramework/Public/EdModeInteractiveToolsContext.h b/Engine/Source/Editor/Experimental/EditorInteractiveToolsFramework/Public/EdModeInteractiveToolsContext.h
--- a/Engine/Source/Editor/Experimental/EditorInteractiveToolsFramework/Public/EdModeInteractiveToolsContext.h
+++ b/Engine/Source/Editor/Experimental/EditorInteractiveToolsFramework/Public/EdModeInteractiveToolsContext.h
@@ -73,9 +73,11 @@ public:
 public:
 	// forwards message to OnToolNotificationMessage delegate
 	virtual void PostToolNotificationMessage(const FText& Message);
+	virtual void PostToolWarningMessage(const FText& Message);
 
 	DECLARE_MULTICAST_DELEGATE_OneParam(FEdModeToolsContextToolNotification, const FText&);
 	FEdModeToolsContextToolNotification OnToolNotificationMessage;
+	FEdModeToolsContextToolNotification OnToolWarningMessage;
 
 
 protected:
diff --git a/Engine/Source/Runtime/Experimental/InteractiveToolsFramework/Public/ToolContextInterfaces.h b/Engine/Source/Runtime/Experimental/InteractiveToolsFramework/Public/ToolContextInterfaces.h
--- a/Engine/Source/Runtime/Experimental/InteractiveToolsFramework/Public/ToolContextInterfaces.h
+++ b/Engine/Source/Runtime/Experimental/InteractiveToolsFramework/Public/ToolContextInterfaces.h
@@ -214,8 +214,10 @@ enum class EToolMessageLevel
 	UserMessage = 1,
 	/** Notification message should be shown in a non-modal notification window */
 	UserNotification = 2,
+	/** Warning message should be shown in a non-modal notification window with panache */
+	UserWarning = 3,
 	/** Error message should be shown in a modal notification window */
-	UserError = 3
+	UserError = 4
 };
 
 
