ModelingTools: reorganize tools into Modeling and Utilities tabs
#rnx

diff --git a/Engine/Plugins/Experimental/ModelingToolsEditorMode/Source/ModelingToolsEditorMode/Private/ModelingToolsEditorMode.cpp b/Engine/Plugins/Experimental/ModelingToolsEditorMode/Source/ModelingToolsEditorMode/Private/ModelingToolsEditorMode.cpp
--- a/Engine/Plugins/Experimental/ModelingToolsEditorMode/Source/ModelingToolsEditorMode/Private/ModelingToolsEditorMode.cpp
+++ b/Engine/Plugins/Experimental/ModelingToolsEditorMode/Source/ModelingToolsEditorMode/Private/ModelingToolsEditorMode.cpp
@@ -445,6 +445,7 @@ void FModelingToolsEditorMode::Enter()
 void FModelingToolsEditorMode::Exit()
 {
 	OnToolNotificationMessage.Clear();
+	OnToolWarningMessage.Clear();
 
 	ToolsContext->ShutdownContext();
 	ToolsContext = nullptr;
diff --git a/Engine/Plugins/Experimental/ModelingToolsEditorMode/Source/ModelingToolsEditorMode/Private/ModelingToolsEditorModeToolkit.cpp b/Engine/Plugins/Experimental/ModelingToolsEditorMode/Source/ModelingToolsEditorMode/Private/ModelingToolsEditorModeToolkit.cpp
--- a/Engine/Plugins/Experimental/ModelingToolsEditorMode/Source/ModelingToolsEditorMode/Private/ModelingToolsEditorModeToolkit.cpp
+++ b/Engine/Plugins/Experimental/ModelingToolsEditorMode/Source/ModelingToolsEditorMode/Private/ModelingToolsEditorModeToolkit.cpp
@@ -226,7 +226,7 @@ UEdModeInteractiveToolsContext* FModelingToolsEditorModeToolkit::GetToolsContext
 }
 
 
-const TArray<FName> FModelingToolsEditorModeToolkit::PaletteNames_Standard = { FName(TEXT("Modeling")) };
+const TArray<FName> FModelingToolsEditorModeToolkit::PaletteNames_Standard = { FName(TEXT("Modeling")), FName(TEXT("Utilities")) };
 const TArray<FName> FModelingToolsEditorModeToolkit::PaletteNames_Experimental = { FName(TEXT("Create")), FName(TEXT("Edit")), FName(TEXT("UVs/Normals")), FName(TEXT("Utilities")) };
 
 
@@ -261,14 +261,15 @@ void FModelingToolsEditorModeToolkit::BuildToolPalette(FName PaletteIndex, class
 void FModelingToolsEditorModeToolkit::BuildToolPalette_Standard(FName PaletteIndex, class FToolBarBuilder& ToolbarBuilder)
 {
 	const FModelingToolsManagerCommands& Commands = FModelingToolsManagerCommands::Get();
-	if (PaletteIndex == FName(TEXT("Modeling")))
-	{
-		ToolbarBuilder.AddToolBarButton(Commands.AcceptActiveTool);
-		ToolbarBuilder.AddToolBarButton(Commands.CancelActiveTool);
-		ToolbarBuilder.AddToolBarButton(Commands.CompleteActiveTool);
 
-		ToolbarBuilder.AddSeparator();
+	ToolbarBuilder.AddToolBarButton(Commands.AcceptActiveTool);
+	ToolbarBuilder.AddToolBarButton(Commands.CancelActiveTool);
+	ToolbarBuilder.AddToolBarButton(Commands.CompleteActiveTool);
+
+	ToolbarBuilder.AddSeparator();
 
+	if (PaletteIndex == FName(TEXT("Modeling")))
+	{
 		ToolbarBuilder.AddToolBarButton(Commands.BeginAddPrimitiveTool);
 		ToolbarBuilder.AddToolBarButton(Commands.BeginDrawPolygonTool);
 
@@ -290,12 +291,13 @@ void FModelingToolsEditorModeToolkit::BuildToolPalette_Standard(FName PaletteInd
 		ToolbarBuilder.AddToolBarButton(Commands.BeginPlaneCutTool);
 		ToolbarBuilder.AddToolBarButton(Commands.BeginSimplifyMeshTool);
 		ToolbarBuilder.AddToolBarButton(Commands.BeginRemeshMeshTool);
-		ToolbarBuilder.AddToolBarButton(Commands.BeginUVProjectionTool);
 		ToolbarBuilder.AddToolBarButton(Commands.BeginVoxelMergeTool);
 		ToolbarBuilder.AddToolBarButton(Commands.BeginVoxelBooleanTool);
 
-		ToolbarBuilder.AddSeparator();
-
+	}
+	else if (PaletteIndex == FName(TEXT("Utilities")))
+	{
+		ToolbarBuilder.AddToolBarButton(Commands.BeginUVProjectionTool);
 		ToolbarBuilder.AddToolBarButton(Commands.BeginEditNormalsTool);
 		ToolbarBuilder.AddToolBarButton(Commands.BeginWeldEdgesTool);
 		ToolbarBuilder.AddToolBarButton(Commands.BeginMeshInspectorTool);
