ModelingTools: icon updates
#rnx

diff --git a/Engine/Plugins/Experimental/ModelingToolsEditorMode/Source/ModelingToolsEditorMode/Private/ModelingToolsManagerActions.cpp b/Engine/Plugins/Experimental/ModelingToolsEditorMode/Source/ModelingToolsEditorMode/Private/ModelingToolsManagerActions.cpp
--- a/Engine/Plugins/Experimental/ModelingToolsEditorMode/Source/ModelingToolsEditorMode/Private/ModelingToolsManagerActions.cpp
+++ b/Engine/Plugins/Experimental/ModelingToolsEditorMode/Source/ModelingToolsEditorMode/Private/ModelingToolsManagerActions.cpp
@@ -40,7 +40,7 @@ void FModelingToolsManagerCommands::RegisterCommands()
 	UI_COMMAND(BeginVoxelMergeTool, "VoxMerge", "Start the Voxel Merge Tool", EUserInterfaceActionType::Button, FInputChord());
 	UI_COMMAND(BeginVoxelBooleanTool, "VoxBoolean", "Start the Voxel Boolean Tool", EUserInterfaceActionType::Button, FInputChord());
 	UI_COMMAND(BeginPlaneCutTool, "PlaneCut", "Start the Plane Cut Tool", EUserInterfaceActionType::Button, FInputChord());
-	UI_COMMAND(BeginPolygonOnMeshTool, "PolygonOnMesh", "Start the Polygon On Mesh Tool", EUserInterfaceActionType::Button, FInputChord());
+	UI_COMMAND(BeginPolygonOnMeshTool, "HoleCut", "Start the Hole Cut Tool", EUserInterfaceActionType::Button, FInputChord());
 	UI_COMMAND(BeginParameterizeMeshTool, "UVGenerate", "Start the Mesh Parameterization Tool", EUserInterfaceActionType::Button, FInputChord());
 	UI_COMMAND(BeginMeshSelectionTool, "Select", "Start the Mesh Selection Tool", EUserInterfaceActionType::Button, FInputChord());
 
