Geometry Modeling Icons

diff --git a/Engine/Plugins/Experimental/ModelingToolsEditorMode/Source/ModelingToolsEditorMode/Private/ModelingToolsEditorModeStyle.cpp b/Engine/Plugins/Experimental/ModelingToolsEditorMode/Source/ModelingToolsEditorMode/Private/ModelingToolsEditorModeStyle.cpp
--- a/Engine/Plugins/Experimental/ModelingToolsEditorMode/Source/ModelingToolsEditorMode/Private/ModelingToolsEditorModeStyle.cpp
+++ b/Engine/Plugins/Experimental/ModelingToolsEditorMode/Source/ModelingToolsEditorMode/Private/ModelingToolsEditorModeStyle.cpp
@@ -143,6 +143,36 @@ void FModelingToolsEditorModeStyle::Initialize()
 		StyleSet->Set("ModelingToolsManagerCommands.BeginAttributeEditorTool", 			new IMAGE_PLUGIN_BRUSH("Icons/AttributeEditor_40x", Icon20x20));
 		StyleSet->Set("ModelingToolsManagerCommands.BeginAttributeEditorTool.Small", 	new IMAGE_PLUGIN_BRUSH("Icons/AttributeEditor_40x", Icon20x20));
 
+		StyleSet->Set("ModelingToolsManagerCommands.BeginAlignObjectsTool",                  new FSlateImageBrush(StyleSet->RootToCoreContentDir(TEXT("../Editor/Slate/Icons/GeneralTools/Align_40x.png")), Icon20x20));
+		StyleSet->Set("ModelingToolsManagerCommands.BeginAlignObjectsTool.Small",            new FSlateImageBrush(StyleSet->RootToCoreContentDir(TEXT("../Editor/Slate/Icons/GeneralTools/Align_40x.png")), Icon20x20));
+		StyleSet->Set("ModelingToolsManagerCommands.BeginGlobalUVGenerateTool",              new IMAGE_PLUGIN_BRUSH("Icons/AutoUnwrap_40x",       Icon20x20));      
+		StyleSet->Set("ModelingToolsManagerCommands.BeginGlobalUVGenerateTool.Small",        new IMAGE_PLUGIN_BRUSH("Icons/AutoUnwrap_40x",       Icon20x20));      
+		StyleSet->Set("ModelingToolsManagerCommands.BeginBakeTransformTool",                 new IMAGE_PLUGIN_BRUSH("Icons/BakeXForm_40x",        Icon20x20));      
+		StyleSet->Set("ModelingToolsManagerCommands.BeginBakeTransformTool.Small",           new IMAGE_PLUGIN_BRUSH("Icons/BakeXForm_40x",        Icon20x20));      
+		StyleSet->Set("ModelingToolsManagerCommands.BeginCombineMeshesTool",                 new IMAGE_PLUGIN_BRUSH("Icons/Combine_40x",          Icon20x20));   
+		StyleSet->Set("ModelingToolsManagerCommands.BeginCombineMeshesTool.Small",           new IMAGE_PLUGIN_BRUSH("Icons/Combine_40x",          Icon20x20));   
+		StyleSet->Set("ModelingToolsManagerCommands.BeginEditMeshMaterialsTool",             new IMAGE_PLUGIN_BRUSH("Icons/EditMats_40x",         Icon20x20));     
+		StyleSet->Set("ModelingToolsManagerCommands.BeginEditMeshMaterialsTool.Small",       new IMAGE_PLUGIN_BRUSH("Icons/EditMats_40x",         Icon20x20));     
+		StyleSet->Set("ModelingToolsManagerCommands.BeginEditPivotTool",                     new IMAGE_PLUGIN_BRUSH("Icons/EditPivot_40x",        Icon20x20));      
+		StyleSet->Set("ModelingToolsManagerCommands.BeginEditPivotTool.Small",               new IMAGE_PLUGIN_BRUSH("Icons/EditPivot_40x",        Icon20x20));      
+		StyleSet->Set("ModelingToolsManagerCommands.BeginGroupUVGenerateTool",               new IMAGE_PLUGIN_BRUSH("Icons/GroupUnwrap_40x",      Icon20x20));       
+		StyleSet->Set("ModelingToolsManagerCommands.BeginGroupUVGenerateTool.Small",         new IMAGE_PLUGIN_BRUSH("Icons/GroupUnwrap_40x",      Icon20x20));       
+		StyleSet->Set("ModelingToolsManagerCommands.BeginRemoveOccludedTrianglesTool",       new IMAGE_PLUGIN_BRUSH("Icons/Jacketing_40x",        Icon20x20));     
+		StyleSet->Set("ModelingToolsManagerCommands.BeginRemoveOccludedTrianglesTool.Small", new IMAGE_PLUGIN_BRUSH("Icons/Jacketing_40x",        Icon20x20));     
+		StyleSet->Set("ModelingToolsManagerCommands.BeginPolygonCutTool",                    new IMAGE_PLUGIN_BRUSH("Icons/PolyCut_40x",          Icon20x20));   
+		StyleSet->Set("ModelingToolsManagerCommands.BeginPolygonCutTool.Small",              new IMAGE_PLUGIN_BRUSH("Icons/PolyCut_40x",          Icon20x20));   
+		StyleSet->Set("ModelingToolsManagerCommands.BeginPolyDeformTool",                    new IMAGE_PLUGIN_BRUSH("Icons/PolyDeform_40x",       Icon20x20));      
+		StyleSet->Set("ModelingToolsManagerCommands.BeginPolyDeformTool.Small",              new IMAGE_PLUGIN_BRUSH("Icons/PolyDeform_40x",       Icon20x20));      
+		StyleSet->Set("ModelingToolsManagerCommands.BeginPolyGroupsTool",                    new IMAGE_PLUGIN_BRUSH("Icons/PolyGroups_40x",       Icon20x20));      
+		StyleSet->Set("ModelingToolsManagerCommands.BeginPolyGroupsTool.Small",              new IMAGE_PLUGIN_BRUSH("Icons/PolyGroups_40x",       Icon20x20));      
+		StyleSet->Set("ModelingToolsManagerCommands.BeginDrawPolyPathTool",                  new IMAGE_PLUGIN_BRUSH("Icons/PolyPath_40x",         Icon20x20));    
+		StyleSet->Set("ModelingToolsManagerCommands.BeginDrawPolyPathTool.Small",            new IMAGE_PLUGIN_BRUSH("Icons/PolyPath_40x",         Icon20x20));    
+		StyleSet->Set("ModelingToolsManagerCommands.BeginMeshSpaceDeformerTool",             new IMAGE_PLUGIN_BRUSH("Icons/SpaceDeform_40x",      Icon20x20));       
+		StyleSet->Set("ModelingToolsManagerCommands.BeginMeshSpaceDeformerTool.Small",       new IMAGE_PLUGIN_BRUSH("Icons/SpaceDeform_40x",      Icon20x20));       
+		StyleSet->Set("ModelingToolsManagerCommands.BeginTransformUVIslandsTool",            new IMAGE_PLUGIN_BRUSH("Icons/TransformUVs_40x",     Icon20x20));         
+		StyleSet->Set("ModelingToolsManagerCommands.BeginTransformUVIslandsTool.Small",      new IMAGE_PLUGIN_BRUSH("Icons/TransformUVs_40x",     Icon20x20));         
+		StyleSet->Set("ModelingToolsManagerCommands.BeginUVLayoutTool",                      new IMAGE_PLUGIN_BRUSH("Icons/UVLayout_40x",         Icon20x20));    
+		StyleSet->Set("ModelingToolsManagerCommands.BeginUVLayoutTool.Small",                new IMAGE_PLUGIN_BRUSH("Icons/UVLayout_40x",         Icon20x20));    
 
 		//const FLinearColor LayerSelectionColor = FLinearColor(0.13f, 0.70f, 1.00f);
 
