Don't show gizmo in Modeling Mode when no actor/component selected. Use new Realtime override functionality to force Realtime mode on in modeling mode.
#rnx

diff --git a/Engine/Plugins/Experimental/ModelingToolsEditorMode/Source/ModelingToolsEditorMode/Private/ModelingToolsEditorMode.cpp b/Engine/Plugins/Experimental/ModelingToolsEditorMode/Source/ModelingToolsEditorMode/Private/ModelingToolsEditorMode.cpp
--- a/Engine/Plugins/Experimental/ModelingToolsEditorMode/Source/ModelingToolsEditorMode/Private/ModelingToolsEditorMode.cpp
+++ b/Engine/Plugins/Experimental/ModelingToolsEditorMode/Source/ModelingToolsEditorMode/Private/ModelingToolsEditorMode.cpp
@@ -104,11 +104,11 @@ bool FModelingToolsEditorMode::CanAutoSave() const
 bool FModelingToolsEditorMode::ShouldDrawWidget() const
 { 
 	// allow standard xform gizmo if we don't have an active tool
-	if (ToolsContext != nullptr)
+	if (ToolsContext != nullptr && ToolsContext->ToolManager->HasAnyActiveTool())
 	{
-		return ToolsContext->ToolManager->HasAnyActiveTool() == false;
+		return false;
 	}
-	return true; 
+	return FEdMode::ShouldDrawWidget(); 
 }
 
 bool FModelingToolsEditorMode::UsesTransformWidget() const
@@ -563,6 +563,9 @@ void FModelingToolsEditorMode::Enter()
 	{
 		FModelingToolActionCommands::UpdateToolCommandBinding(Tool, UICommandList, true);
 	});
+
+	// enable realtime viewport override
+	ConfigureRealTimeViewportsOverride(true);
 }
 
 
@@ -584,6 +587,9 @@ void FModelingToolsEditorMode::Exit()
 		Toolkit.Reset();
 	}
 
+	// clear realtime viewport override
+	ConfigureRealTimeViewportsOverride(false);
+
 	// Call base Exit method to ensure proper cleanup
 	FEdMode::Exit();
 }
@@ -640,4 +646,32 @@ bool FModelingToolsEditorMode::GetPivotForOrbit(FVector& OutPivot) const
 }
 
 
+
+void FModelingToolsEditorMode::ConfigureRealTimeViewportsOverride(bool bEnable)
+{
+	FLevelEditorModule& LevelEditorModule = FModuleManager::GetModuleChecked<FLevelEditorModule>("LevelEditor");
+	TSharedPtr<ILevelEditor> LevelEditor = LevelEditorModule.GetFirstLevelEditor();
+	if (LevelEditor.IsValid())
+	{
+		TArray<TSharedPtr<IAssetViewport>> Viewports = LevelEditor->GetViewports();
+		for (const TSharedPtr<IAssetViewport>& ViewportWindow : Viewports)
+		{
+			if (ViewportWindow.IsValid())
+			{
+				FEditorViewportClient& Viewport = ViewportWindow->GetAssetViewportClient();
+				if (bEnable)
+				{
+					Viewport.SetRealtimeOverride(bEnable, LOCTEXT("RealtimeOverrideMessage_ModelingMode", "Modeling Mode"));
+				}
+				else
+				{
+					Viewport.RemoveRealtimeOverride();
+				}
+			}
+		}
+	}
+}
+
+
+
 #undef LOCTEXT_NAMESPACE
\ No newline at end of file
diff --git a/Engine/Plugins/Experimental/ModelingToolsEditorMode/Source/ModelingToolsEditorMode/Public/ModelingToolsEditorMode.h b/Engine/Plugins/Experimental/ModelingToolsEditorMode/Source/ModelingToolsEditorMode/Public/ModelingToolsEditorMode.h
--- a/Engine/Plugins/Experimental/ModelingToolsEditorMode/Source/ModelingToolsEditorMode/Public/ModelingToolsEditorMode.h
+++ b/Engine/Plugins/Experimental/ModelingToolsEditorMode/Source/ModelingToolsEditorMode/Public/ModelingToolsEditorMode.h
@@ -161,4 +161,6 @@ protected:
 
 	void ModelingModeShortcutRequested(EModelingModeActionCommands Command);
 	void FocusCameraAtCursorHotkey();
+
+	void ConfigureRealTimeViewportsOverride(bool bEnable);
 };
