Disable prototypes tab in Modeling Tools mode
#rnx

diff --git a/Engine/Plugins/Experimental/ModelingToolsEditorMode/Source/ModelingToolsEditorMode/Private/ModelingToolsEditorModeToolkit.cpp b/Engine/Plugins/Experimental/ModelingToolsEditorMode/Source/ModelingToolsEditorMode/Private/ModelingToolsEditorModeToolkit.cpp
--- a/Engine/Plugins/Experimental/ModelingToolsEditorMode/Source/ModelingToolsEditorMode/Private/ModelingToolsEditorModeToolkit.cpp
+++ b/Engine/Plugins/Experimental/ModelingToolsEditorMode/Source/ModelingToolsEditorMode/Private/ModelingToolsEditorModeToolkit.cpp
@@ -26,7 +26,7 @@
 // if set to 1, then on mode initialization we include buttons for prototype modeling tools
 static TAutoConsoleVariable<int32> CVarEnablePrototypeModelingTools(
 	TEXT("modeling.EnablePrototypes"),
-	1,
+	0,
 	TEXT("Enable unsupported Experimental prototype Modeling Tools"));
 
 
@@ -327,7 +327,7 @@ static const FName DeformTabName(TEXT("Deform"));
 static const FName PrototypesTabName(TEXT("Prototypes"));
 
 
-const TArray<FName> FModelingToolsEditorModeToolkit::PaletteNames_Standard = { FName(TEXT("Modeling")), FName(TEXT("Utilities")) };
+const TArray<FName> FModelingToolsEditorModeToolkit::PaletteNames_Standard = { CreateTabName, TransformTabName, DeformTabName, PolyGroupsTabName, TrianglesTabName, UVNormalTabName };
 const TArray<FName> FModelingToolsEditorModeToolkit::PaletteNames_Experimental = { CreateTabName, TransformTabName, DeformTabName, PolyGroupsTabName, TrianglesTabName, UVNormalTabName, PrototypesTabName };
 
 
@@ -369,42 +369,93 @@ void FModelingToolsEditorModeToolkit::BuildToolPalette_Standard(FName PaletteInd
 
 	ToolbarBuilder.AddSeparator();
 
-	if (PaletteIndex == FName(TEXT("Modeling")))
+
+	if (PaletteIndex == CreateTabName)
 	{
+		ToolbarBuilder.AddToolBarButton(Commands.BeginTransformMeshesTool);
+		ToolbarBuilder.AddToolBarButton(Commands.BeginMeshSelectionTool);
+		ToolbarBuilder.AddSeparator();
 		ToolbarBuilder.AddToolBarButton(Commands.BeginAddPrimitiveTool);
 		ToolbarBuilder.AddToolBarButton(Commands.BeginDrawPolygonTool);
-
+		ToolbarBuilder.AddToolBarButton(Commands.BeginDrawPolyPathTool);
 		ToolbarBuilder.AddSeparator();
-
+		ToolbarBuilder.AddToolBarButton(Commands.BeginVoxelMergeTool);
+		ToolbarBuilder.AddToolBarButton(Commands.BeginVoxelBooleanTool);
+	}
+	else if (PaletteIndex == TransformTabName)
+	{
+		ToolbarBuilder.AddToolBarButton(Commands.BeginTransformMeshesTool);
+		ToolbarBuilder.AddToolBarButton(Commands.BeginMeshSelectionTool);
+		ToolbarBuilder.AddSeparator();
+		ToolbarBuilder.AddToolBarButton(Commands.BeginAlignObjectsTool);
+		ToolbarBuilder.AddToolBarButton(Commands.BeginEditPivotTool);
+		ToolbarBuilder.AddToolBarButton(Commands.BeginBakeTransformTool);
+		ToolbarBuilder.AddToolBarButton(Commands.BeginCombineMeshesTool);
+		ToolbarBuilder.AddSeparator();
+		ToolbarBuilder.AddToolBarButton(Commands.BeginMeshInspectorTool);
+	}
+	else if (PaletteIndex == DeformTabName)
+	{
 		ToolbarBuilder.AddToolBarButton(Commands.BeginTransformMeshesTool);
 		ToolbarBuilder.AddToolBarButton(Commands.BeginMeshSelectionTool);
-		ToolbarBuilder.AddToolBarButton(Commands.BeginPolyEditTool);
-		ToolbarBuilder.AddToolBarButton(Commands.BeginSmoothMeshTool);
-		ToolbarBuilder.AddToolBarButton(Commands.BeginDisplaceMeshTool);
-
 		ToolbarBuilder.AddSeparator();
-
 		ToolbarBuilder.AddToolBarButton(Commands.BeginSculptMeshTool);
 		ToolbarBuilder.AddToolBarButton(Commands.BeginRemeshSculptMeshTool);
-
 		ToolbarBuilder.AddSeparator();
-
+		ToolbarBuilder.AddToolBarButton(Commands.BeginPlaneCutTool);
+		ToolbarBuilder.AddToolBarButton(Commands.BeginPolygonCutTool);
+		ToolbarBuilder.AddToolBarButton(Commands.BeginSmoothMeshTool);
+		ToolbarBuilder.AddToolBarButton(Commands.BeginDisplaceMeshTool);
+		ToolbarBuilder.AddToolBarButton(Commands.BeginMeshSpaceDeformerTool);
+		ToolbarBuilder.AddSeparator();
+		ToolbarBuilder.AddToolBarButton(Commands.BeginRemeshMeshTool);
+		ToolbarBuilder.AddToolBarButton(Commands.BeginMeshInspectorTool);
+	}
+	else if (PaletteIndex == TrianglesTabName)
+	{
+		ToolbarBuilder.AddToolBarButton(Commands.BeginTransformMeshesTool);
+		ToolbarBuilder.AddToolBarButton(Commands.BeginMeshSelectionTool);
+		ToolbarBuilder.AddSeparator();
+		ToolbarBuilder.AddToolBarButton(Commands.BeginTriEditTool);
 		ToolbarBuilder.AddToolBarButton(Commands.BeginSimplifyMeshTool);
 		ToolbarBuilder.AddToolBarButton(Commands.BeginRemeshMeshTool);
-		ToolbarBuilder.AddToolBarButton(Commands.BeginPlaneCutTool);
-		ToolbarBuilder.AddToolBarButton(Commands.BeginVoxelMergeTool);
-		ToolbarBuilder.AddToolBarButton(Commands.BeginVoxelBooleanTool);
-
+		ToolbarBuilder.AddSeparator();
+		ToolbarBuilder.AddToolBarButton(Commands.BeginEditMeshMaterialsTool);
+		ToolbarBuilder.AddToolBarButton(Commands.BeginWeldEdgesTool);
+		ToolbarBuilder.AddToolBarButton(Commands.BeginRemoveOccludedTrianglesTool);
+		ToolbarBuilder.AddSeparator();
+		ToolbarBuilder.AddToolBarButton(Commands.BeginMeshInspectorTool);
+	}
+	else if (PaletteIndex == PolyGroupsTabName)
+	{
+		ToolbarBuilder.AddToolBarButton(Commands.BeginTransformMeshesTool);
+		ToolbarBuilder.AddToolBarButton(Commands.BeginMeshSelectionTool);
+		ToolbarBuilder.AddToolBarButton(Commands.BeginPolyGroupsTool);
+		ToolbarBuilder.AddSeparator();
+		ToolbarBuilder.AddToolBarButton(Commands.BeginPolyEditTool);
+		ToolbarBuilder.AddToolBarButton(Commands.BeginPolyDeformTool);
+		ToolbarBuilder.AddSeparator();
+		ToolbarBuilder.AddToolBarButton(Commands.BeginAttributeEditorTool);
+		ToolbarBuilder.AddToolBarButton(Commands.BeginMeshInspectorTool);
 	}
-	else if (PaletteIndex == FName(TEXT("Utilities")))
+	else if (PaletteIndex == UVNormalTabName)
 	{
+		ToolbarBuilder.AddToolBarButton(Commands.BeginTransformMeshesTool);
+		ToolbarBuilder.AddToolBarButton(Commands.BeginMeshSelectionTool);
+		ToolbarBuilder.AddToolBarButton(Commands.BeginPolyGroupsTool);
+		ToolbarBuilder.AddSeparator();
+		ToolbarBuilder.AddToolBarButton(Commands.BeginEditNormalsTool);
+		ToolbarBuilder.AddSeparator();
+		ToolbarBuilder.AddToolBarButton(Commands.BeginGlobalUVGenerateTool);
+		ToolbarBuilder.AddToolBarButton(Commands.BeginGroupUVGenerateTool);
 		ToolbarBuilder.AddToolBarButton(Commands.BeginUVProjectionTool);
+		ToolbarBuilder.AddToolBarButton(Commands.BeginTransformUVIslandsTool);
 		ToolbarBuilder.AddToolBarButton(Commands.BeginUVLayoutTool);
-		ToolbarBuilder.AddToolBarButton(Commands.BeginEditNormalsTool);
-		ToolbarBuilder.AddToolBarButton(Commands.BeginWeldEdgesTool);
-		ToolbarBuilder.AddToolBarButton(Commands.BeginMeshInspectorTool);
+		ToolbarBuilder.AddSeparator();
 		ToolbarBuilder.AddToolBarButton(Commands.BeginAttributeEditorTool);
+		ToolbarBuilder.AddToolBarButton(Commands.BeginMeshInspectorTool);
 	}
+
 }
 
 
