Add file missed in last commit
#rnx

diff --git a/Engine/Plugins/Experimental/ModelingToolsEditorMode/Source/ModelingToolsEditorMode/Public/ModelingToolsEditorMode.h b/Engine/Plugins/Experimental/ModelingToolsEditorMode/Source/ModelingToolsEditorMode/Public/ModelingToolsEditorMode.h
--- a/Engine/Plugins/Experimental/ModelingToolsEditorMode/Source/ModelingToolsEditorMode/Public/ModelingToolsEditorMode.h
+++ b/Engine/Plugins/Experimental/ModelingToolsEditorMode/Source/ModelingToolsEditorMode/Public/ModelingToolsEditorMode.h
@@ -39,7 +39,8 @@ public:
 
 	virtual void ActorSelectionChangeNotify() override;
 
-	virtual bool ProcessEditDelete();
+	virtual bool ProcessEditDelete() override;
+	virtual bool ProcessEditCut() override;
 
 	virtual bool CanAutoSave() const override;
 
