ModelingTools: Add Duplicate option to toolbar, a variant of Combine tool that works with a single input
#rnx

diff --git a/Engine/Plugins/Experimental/MeshModelingToolset/Source/MeshModelingTools/Private/CombineMeshesTool.cpp b/Engine/Plugins/Experimental/MeshModelingToolset/Source/MeshModelingTools/Private/CombineMeshesTool.cpp
--- a/Engine/Plugins/Experimental/MeshModelingToolset/Source/MeshModelingTools/Private/CombineMeshesTool.cpp
+++ b/Engine/Plugins/Experimental/MeshModelingToolset/Source/MeshModelingTools/Private/CombineMeshesTool.cpp
@@ -32,7 +32,9 @@
 
 bool UCombineMeshesToolBuilder::CanBuildTool(const FToolBuilderState& SceneState) const
 {
-	return AssetAPI != nullptr && ToolBuilderUtil::CountComponents(SceneState, CanMakeComponentTarget) > 1;
+	return (bIsDuplicateTool) ?
+		  (AssetAPI != nullptr && ToolBuilderUtil::CountComponents(SceneState, CanMakeComponentTarget) == 1)
+		: (AssetAPI != nullptr && ToolBuilderUtil::CountComponents(SceneState, CanMakeComponentTarget) > 1);
 }
 
 UInteractiveTool* UCombineMeshesToolBuilder::BuildTool(const FToolBuilderState& SceneState) const
@@ -40,7 +42,7 @@ UInteractiveTool* UCombineMeshesToolBuilder::BuildTool(const FToolBuilderState&
 	UCombineMeshesTool* NewTool = NewObject<UCombineMeshesTool>(SceneState.ToolManager);
 
 	TArray<UActorComponent*> Components = ToolBuilderUtil::FindAllComponents(SceneState, CanMakeComponentTarget);
-	check(Components.Num() > 1);
+	check(Components.Num() > 0);
 
 	TArray<TUniquePtr<FPrimitiveComponentTarget>> ComponentTargets;
 	for (UActorComponent* ActorComponent : Components)
@@ -55,6 +57,7 @@ UInteractiveTool* UCombineMeshesToolBuilder::BuildTool(const FToolBuilderState&
 	NewTool->SetSelection(MoveTemp(ComponentTargets));
 	NewTool->SetWorld(SceneState.World);
 	NewTool->SetAssetAPI(AssetAPI);
+	NewTool->SetDuplicateMode(bIsDuplicateTool);
 
 	return NewTool;
 }
@@ -81,6 +84,11 @@ void UCombineMeshesTool::SetWorld(UWorld* World)
 	this->TargetWorld = World;
 }
 
+void UCombineMeshesTool::SetDuplicateMode(bool bDuplicateModeIn)
+{
+	this->bDuplicateMode = bDuplicateModeIn;
+}
+
 void UCombineMeshesTool::Setup()
 {
 	UInteractiveTool::Setup();
diff --git a/Engine/Plugins/Experimental/MeshModelingToolset/Source/MeshModelingTools/Public/CombineMeshesTool.h b/Engine/Plugins/Experimental/MeshModelingToolset/Source/MeshModelingTools/Public/CombineMeshesTool.h
--- a/Engine/Plugins/Experimental/MeshModelingToolset/Source/MeshModelingTools/Public/CombineMeshesTool.h
+++ b/Engine/Plugins/Experimental/MeshModelingToolset/Source/MeshModelingTools/Public/CombineMeshesTool.h
@@ -30,6 +30,7 @@ class MESHMODELINGTOOLS_API UCombineMeshesToolBuilder : public UInteractiveToolB
 
 public:
 	IToolsContextAssetAPI* AssetAPI = nullptr;
+	bool bIsDuplicateTool = false;
 
 	virtual bool CanBuildTool(const FToolBuilderState& SceneState) const override;
 	virtual UInteractiveTool* BuildTool(const FToolBuilderState& SceneState) const override;
@@ -67,6 +68,8 @@ class MESHMODELINGTOOLS_API UCombineMeshesTool : public UMultiSelectionTool
 public:
 	UCombineMeshesTool();
 
+	virtual void SetDuplicateMode(bool bDuplicateMode);
+
 	virtual void Setup() override;
 	virtual void Shutdown(EToolShutdownType ShutdownType) override;
 
@@ -89,5 +92,7 @@ protected:
 	UWorld* TargetWorld;
 	IToolsContextAssetAPI* AssetAPI;
 
+	bool bDuplicateMode;
+
 	void UpdateAssets();
 };
diff --git a/Engine/Plugins/Experimental/ModelingToolsEditorMode/Source/ModelingToolsEditorMode/Private/ModelingToolsEditorMode.cpp b/Engine/Plugins/Experimental/ModelingToolsEditorMode/Source/ModelingToolsEditorMode/Private/ModelingToolsEditorMode.cpp
--- a/Engine/Plugins/Experimental/ModelingToolsEditorMode/Source/ModelingToolsEditorMode/Private/ModelingToolsEditorMode.cpp
+++ b/Engine/Plugins/Experimental/ModelingToolsEditorMode/Source/ModelingToolsEditorMode/Private/ModelingToolsEditorMode.cpp
@@ -503,6 +503,12 @@ void FModelingToolsEditorMode::Enter()
 	CombineMeshesToolBuilder->AssetAPI = ToolsContext->GetAssetAPI();
 	RegisterToolFunc(ToolManagerCommands.BeginCombineMeshesTool, TEXT("CombineMeshesTool"), CombineMeshesToolBuilder);
 
+	UCombineMeshesToolBuilder* DuplicateMeshesToolBuilder = NewObject<UCombineMeshesToolBuilder>();
+	DuplicateMeshesToolBuilder->AssetAPI = ToolsContext->GetAssetAPI();
+	DuplicateMeshesToolBuilder->bIsDuplicateTool = true;
+	RegisterToolFunc(ToolManagerCommands.BeginDuplicateMeshesTool, TEXT("DuplicateMeshesTool"), DuplicateMeshesToolBuilder);
+
+
 	// edit tools
 
 
diff --git a/Engine/Plugins/Experimental/ModelingToolsEditorMode/Source/ModelingToolsEditorMode/Private/ModelingToolsEditorModeStyle.cpp b/Engine/Plugins/Experimental/ModelingToolsEditorMode/Source/ModelingToolsEditorMode/Private/ModelingToolsEditorModeStyle.cpp
--- a/Engine/Plugins/Experimental/ModelingToolsEditorMode/Source/ModelingToolsEditorMode/Private/ModelingToolsEditorModeStyle.cpp
+++ b/Engine/Plugins/Experimental/ModelingToolsEditorMode/Source/ModelingToolsEditorMode/Private/ModelingToolsEditorModeStyle.cpp
@@ -106,6 +106,8 @@ void FModelingToolsEditorModeStyle::Initialize()
 		StyleSet->Set("ModelingToolsManagerCommands.BeginSculptMeshTool.Small", 		new IMAGE_PLUGIN_BRUSH("Icons/Sculpt_40x", 			Icon20x20));
 		StyleSet->Set("ModelingToolsManagerCommands.BeginPolyEditTool", 				new IMAGE_PLUGIN_BRUSH("Icons/PolyEdit_40x", 		Icon20x20));
 		StyleSet->Set("ModelingToolsManagerCommands.BeginPolyEditTool.Small", 			new IMAGE_PLUGIN_BRUSH("Icons/PolyEdit_40x", 		Icon20x20));
+		StyleSet->Set("ModelingToolsManagerCommands.BeginTriEditTool", 					new IMAGE_PLUGIN_BRUSH("Icons/TriEdit_40x", 		Icon20x20));
+		StyleSet->Set("ModelingToolsManagerCommands.BeginTriEditTool.Small", 			new IMAGE_PLUGIN_BRUSH("Icons/TriEdit_40x", 		Icon20x20));
 		//StyleSet->Set("ModelingToolsManagerCommands.BeginPolyDeformTool", 				new IMAGE_PLUGIN_BRUSH("Icons/PolyEdit_40x", 		Icon20x20));
 		//StyleSet->Set("ModelingToolsManagerCommands.BeginPolyDeformTool.Small", 		new IMAGE_PLUGIN_BRUSH("Icons/PolyEdit_40x", 		Icon20x20));
 		StyleSet->Set("ModelingToolsManagerCommands.BeginDisplaceMeshTool", 			new IMAGE_PLUGIN_BRUSH("Icons/Displace_40x", 		Icon20x20));
@@ -151,6 +153,8 @@ void FModelingToolsEditorModeStyle::Initialize()
 		StyleSet->Set("ModelingToolsManagerCommands.BeginBakeTransformTool.Small",           new IMAGE_PLUGIN_BRUSH("Icons/BakeXForm_40x",        Icon20x20));      
 		StyleSet->Set("ModelingToolsManagerCommands.BeginCombineMeshesTool",                 new IMAGE_PLUGIN_BRUSH("Icons/Combine_40x",          Icon20x20));   
 		StyleSet->Set("ModelingToolsManagerCommands.BeginCombineMeshesTool.Small",           new IMAGE_PLUGIN_BRUSH("Icons/Combine_40x",          Icon20x20));   
+		StyleSet->Set("ModelingToolsManagerCommands.BeginDuplicateMeshesTool",               new IMAGE_PLUGIN_BRUSH("Icons/Combine_40x",          Icon20x20));   
+		StyleSet->Set("ModelingToolsManagerCommands.BeginDuplicateMeshesTool.Small",         new IMAGE_PLUGIN_BRUSH("Icons/Combine_40x",          Icon20x20));   
 		StyleSet->Set("ModelingToolsManagerCommands.BeginEditMeshMaterialsTool",             new IMAGE_PLUGIN_BRUSH("Icons/EditMats_40x",         Icon20x20));     
 		StyleSet->Set("ModelingToolsManagerCommands.BeginEditMeshMaterialsTool.Small",       new IMAGE_PLUGIN_BRUSH("Icons/EditMats_40x",         Icon20x20));     
 		StyleSet->Set("ModelingToolsManagerCommands.BeginEditPivotTool",                     new IMAGE_PLUGIN_BRUSH("Icons/EditPivot_40x",        Icon20x20));      
diff --git a/Engine/Plugins/Experimental/ModelingToolsEditorMode/Source/ModelingToolsEditorMode/Private/ModelingToolsEditorModeToolkit.cpp b/Engine/Plugins/Experimental/ModelingToolsEditorMode/Source/ModelingToolsEditorMode/Private/ModelingToolsEditorModeToolkit.cpp
--- a/Engine/Plugins/Experimental/ModelingToolsEditorMode/Source/ModelingToolsEditorMode/Private/ModelingToolsEditorModeToolkit.cpp
+++ b/Engine/Plugins/Experimental/ModelingToolsEditorMode/Source/ModelingToolsEditorMode/Private/ModelingToolsEditorModeToolkit.cpp
@@ -391,6 +391,7 @@ void FModelingToolsEditorModeToolkit::BuildToolPalette_Standard(FName PaletteInd
 		ToolbarBuilder.AddToolBarButton(Commands.BeginEditPivotTool);
 		ToolbarBuilder.AddToolBarButton(Commands.BeginBakeTransformTool);
 		ToolbarBuilder.AddToolBarButton(Commands.BeginCombineMeshesTool);
+		ToolbarBuilder.AddToolBarButton(Commands.BeginDuplicateMeshesTool);
 		ToolbarBuilder.AddSeparator();
 		ToolbarBuilder.AddToolBarButton(Commands.BeginMeshInspectorTool);
 	}
@@ -491,6 +492,7 @@ void FModelingToolsEditorModeToolkit::BuildToolPalette_Experimental(FName Palett
 		ToolbarBuilder.AddToolBarButton(Commands.BeginEditPivotTool);
 		ToolbarBuilder.AddToolBarButton(Commands.BeginBakeTransformTool);
 		ToolbarBuilder.AddToolBarButton(Commands.BeginCombineMeshesTool);
+		ToolbarBuilder.AddToolBarButton(Commands.BeginDuplicateMeshesTool);
 		ToolbarBuilder.AddSeparator();
 		ToolbarBuilder.AddToolBarButton(Commands.BeginMeshInspectorTool);
 	}
diff --git a/Engine/Plugins/Experimental/ModelingToolsEditorMode/Source/ModelingToolsEditorMode/Private/ModelingToolsManagerActions.cpp b/Engine/Plugins/Experimental/ModelingToolsEditorMode/Source/ModelingToolsEditorMode/Private/ModelingToolsManagerActions.cpp
--- a/Engine/Plugins/Experimental/ModelingToolsEditorMode/Source/ModelingToolsEditorMode/Private/ModelingToolsManagerActions.cpp
+++ b/Engine/Plugins/Experimental/ModelingToolsEditorMode/Source/ModelingToolsEditorMode/Private/ModelingToolsManagerActions.cpp
@@ -37,7 +37,8 @@ void FModelingToolsManagerCommands::RegisterCommands()
 	UI_COMMAND(BeginTransformMeshesTool, "Transform", "Start the Transform Tool", EUserInterfaceActionType::Button, FInputChord());
 	UI_COMMAND(BeginEditPivotTool, "Edit Pivot", "Start the Edit Pivot Tool", EUserInterfaceActionType::Button, FInputChord());
 	UI_COMMAND(BeginBakeTransformTool, "Bake XForm", "Start the Bake Transform Tool", EUserInterfaceActionType::Button, FInputChord());
-	UI_COMMAND(BeginCombineMeshesTool, "Combine", "Start the Combine Meshes Tool", EUserInterfaceActionType::Button, FInputChord());
+	UI_COMMAND(BeginCombineMeshesTool, "Combine", "Combine the selected Mesh Assets into a new Asset", EUserInterfaceActionType::Button, FInputChord());
+	UI_COMMAND(BeginDuplicateMeshesTool, "Duplicate", "Duplicate the selected Mesh Asset into a new Asset", EUserInterfaceActionType::Button, FInputChord());
 	UI_COMMAND(BeginAlignObjectsTool, "Align", "Start the Align Objects Tool", EUserInterfaceActionType::Button, FInputChord());
 
 	UI_COMMAND(BeginRemeshSculptMeshTool, "DynaSculpt", "Start the Liquid Sculpt Mesh Tool", EUserInterfaceActionType::Button, FInputChord());
diff --git a/Engine/Plugins/Experimental/ModelingToolsEditorMode/Source/ModelingToolsEditorMode/Public/ModelingToolsManagerActions.h b/Engine/Plugins/Experimental/ModelingToolsEditorMode/Source/ModelingToolsEditorMode/Public/ModelingToolsManagerActions.h
--- a/Engine/Plugins/Experimental/ModelingToolsEditorMode/Source/ModelingToolsEditorMode/Public/ModelingToolsManagerActions.h
+++ b/Engine/Plugins/Experimental/ModelingToolsEditorMode/Source/ModelingToolsEditorMode/Public/ModelingToolsManagerActions.h
@@ -32,6 +32,7 @@ public:
 	TSharedPtr<FUICommandInfo> BeginEditPivotTool;
 	TSharedPtr<FUICommandInfo> BeginBakeTransformTool;
 	TSharedPtr<FUICommandInfo> BeginCombineMeshesTool;
+	TSharedPtr<FUICommandInfo> BeginDuplicateMeshesTool;
 	TSharedPtr<FUICommandInfo> BeginAlignObjectsTool;
 
 	TSharedPtr<FUICommandInfo> BeginRemeshSculptMeshTool;
