ModelingTools: Update icons and shaders
#rnx

diff --git a/Engine/Plugins/Experimental/MeshModelingToolset/Source/MeshModelingToolsEditorOnly/Private/MergeMeshesTool.cpp b/Engine/Plugins/Experimental/MeshModelingToolset/Source/MeshModelingToolsEditorOnly/Private/MergeMeshesTool.cpp
--- a/Engine/Plugins/Experimental/MeshModelingToolset/Source/MeshModelingToolsEditorOnly/Private/MergeMeshesTool.cpp
+++ b/Engine/Plugins/Experimental/MeshModelingToolset/Source/MeshModelingToolsEditorOnly/Private/MergeMeshesTool.cpp
@@ -103,8 +103,7 @@ void UMergeMeshesTool::Setup()
 	CreateLowQualityPreview(); // update the preview with a low-quality result
 	
 	Preview->ConfigureMaterials(
-		// TODO: This is is very likely a bug waiting to happen
-		ToolSetupUtil::GetDefaultMaterial(GetToolManager(), ComponentTargets[0]->GetMaterial(0)),
+		ToolSetupUtil::GetDefaultSculptMaterial(GetToolManager()),
 		ToolSetupUtil::GetDefaultWorkingMaterial(GetToolManager())
 	);
 
diff --git a/Engine/Plugins/Experimental/MeshModelingToolset/Source/MeshModelingToolsEditorOnly/Private/VoxelCSGMeshesTool.cpp b/Engine/Plugins/Experimental/MeshModelingToolset/Source/MeshModelingToolsEditorOnly/Private/VoxelCSGMeshesTool.cpp
--- a/Engine/Plugins/Experimental/MeshModelingToolset/Source/MeshModelingToolsEditorOnly/Private/VoxelCSGMeshesTool.cpp
+++ b/Engine/Plugins/Experimental/MeshModelingToolset/Source/MeshModelingToolsEditorOnly/Private/VoxelCSGMeshesTool.cpp
@@ -101,8 +101,7 @@ void UVoxelCSGMeshesTool::Setup()
 	CreateLowQualityPreview();
 
 	Preview->ConfigureMaterials(
-		// TODO : Probably multi-select bug - MCD
-		ToolSetupUtil::GetDefaultMaterial(GetToolManager(), ComponentTargets[0]->GetMaterial(0)),
+		ToolSetupUtil::GetDefaultSculptMaterial(GetToolManager()),
 		ToolSetupUtil::GetDefaultWorkingMaterial(GetToolManager())
 	);
 	
diff --git a/Engine/Plugins/Experimental/ModelingToolsEditorMode/Source/ModelingToolsEditorMode/Private/ModelingToolsEditorModeStyle.cpp b/Engine/Plugins/Experimental/ModelingToolsEditorMode/Source/ModelingToolsEditorMode/Private/ModelingToolsEditorModeStyle.cpp
--- a/Engine/Plugins/Experimental/ModelingToolsEditorMode/Source/ModelingToolsEditorMode/Private/ModelingToolsEditorModeStyle.cpp
+++ b/Engine/Plugins/Experimental/ModelingToolsEditorMode/Source/ModelingToolsEditorMode/Private/ModelingToolsEditorModeStyle.cpp
@@ -82,8 +82,8 @@ void FModelingToolsEditorModeStyle::Initialize()
 		StyleSet->Set("ModelingToolsManagerCommands.CompleteActiveTool.Small", new IMAGE_PLUGIN_BRUSH("Icons/icon_ActiveTool_Accept_40x", Icon20x20));
 
 
-		StyleSet->Set("ModelingToolsManagerCommands.BeginShapeSprayTool", 				new IMAGE_PLUGIN_BRUSH("Icons/icon_Tool_ShapeSpray_40x",	Icon20x20));
-		StyleSet->Set("ModelingToolsManagerCommands.BeginShapeSprayTool.Small", 		new IMAGE_PLUGIN_BRUSH("Icons/icon_Tool_ShapeSpray_40x",	Icon20x20));
+		StyleSet->Set("ModelingToolsManagerCommands.BeginShapeSprayTool", 				new IMAGE_PLUGIN_BRUSH("Icons/ShapeSpray_40x",	Icon20x20));
+		StyleSet->Set("ModelingToolsManagerCommands.BeginShapeSprayTool.Small", 		new IMAGE_PLUGIN_BRUSH("Icons/ShapeSpray_40x",	Icon20x20));
 		//StyleSet->Set("ModelingToolsManagerCommands.BeginMeshSpaceDeformerTool", 		new IMAGE_PLUGIN_BRUSH("Icons/icon_Tool_Displace_40x",		Icon20x20));
 		//StyleSet->Set("ModelingToolsManagerCommands.BeginMeshSpaceDeformerTool.Small", 	new IMAGE_PLUGIN_BRUSH("Icons/icon_Tool_Displace_40x",		Icon20x20));
 		StyleSet->Set("ModelingToolsManagerCommands.BeginPolygonOnMeshTool", 			new IMAGE_PLUGIN_BRUSH("Icons/icon_Tool_PolygonOnMesh_40x",	Icon20x20));
@@ -99,6 +99,8 @@ void FModelingToolsEditorModeStyle::Initialize()
 		StyleSet->Set("ModelingToolsManagerCommands.BeginAddPrimitiveTool.Small", 		new IMAGE_PLUGIN_BRUSH("Icons/Primitive_40x",		Icon20x20));
 		StyleSet->Set("ModelingToolsManagerCommands.BeginDrawPolygonTool", 				new IMAGE_PLUGIN_BRUSH("Icons/DrawPolygon_40x",		Icon20x20));
 		StyleSet->Set("ModelingToolsManagerCommands.BeginDrawPolygonTool.Small", 		new IMAGE_PLUGIN_BRUSH("Icons/DrawPolygon_40x", 	Icon20x20));
+		StyleSet->Set("ModelingToolsManagerCommands.BeginAddPatchTool",					new IMAGE_PLUGIN_BRUSH("Icons/Patch_40x",			Icon20x20));
+		StyleSet->Set("ModelingToolsManagerCommands.BeginAddPatchTool.Small",			new IMAGE_PLUGIN_BRUSH("Icons/Patch_40x",			Icon20x20));
 
 		StyleSet->Set("ModelingToolsManagerCommands.BeginSmoothMeshTool", 				new IMAGE_PLUGIN_BRUSH("Icons/Smooth_40x", 			Icon20x20));
 		StyleSet->Set("ModelingToolsManagerCommands.BeginSmoothMeshTool.Small", 		new IMAGE_PLUGIN_BRUSH("Icons/Smooth_40x", 			Icon20x20));
@@ -153,8 +155,8 @@ void FModelingToolsEditorModeStyle::Initialize()
 		StyleSet->Set("ModelingToolsManagerCommands.BeginBakeTransformTool.Small",           new IMAGE_PLUGIN_BRUSH("Icons/BakeXForm_40x",        Icon20x20));      
 		StyleSet->Set("ModelingToolsManagerCommands.BeginCombineMeshesTool",                 new IMAGE_PLUGIN_BRUSH("Icons/Combine_40x",          Icon20x20));   
 		StyleSet->Set("ModelingToolsManagerCommands.BeginCombineMeshesTool.Small",           new IMAGE_PLUGIN_BRUSH("Icons/Combine_40x",          Icon20x20));   
-		StyleSet->Set("ModelingToolsManagerCommands.BeginDuplicateMeshesTool",               new IMAGE_PLUGIN_BRUSH("Icons/Combine_40x",          Icon20x20));   
-		StyleSet->Set("ModelingToolsManagerCommands.BeginDuplicateMeshesTool.Small",         new IMAGE_PLUGIN_BRUSH("Icons/Combine_40x",          Icon20x20));   
+		StyleSet->Set("ModelingToolsManagerCommands.BeginDuplicateMeshesTool",               new IMAGE_PLUGIN_BRUSH("Icons/Duplicate_40x",        Icon20x20));   
+		StyleSet->Set("ModelingToolsManagerCommands.BeginDuplicateMeshesTool.Small",         new IMAGE_PLUGIN_BRUSH("Icons/Duplicate_40x",        Icon20x20));   
 		StyleSet->Set("ModelingToolsManagerCommands.BeginEditMeshMaterialsTool",             new IMAGE_PLUGIN_BRUSH("Icons/EditMats_40x",         Icon20x20));     
 		StyleSet->Set("ModelingToolsManagerCommands.BeginEditMeshMaterialsTool.Small",       new IMAGE_PLUGIN_BRUSH("Icons/EditMats_40x",         Icon20x20));     
 		StyleSet->Set("ModelingToolsManagerCommands.BeginEditPivotTool",                     new IMAGE_PLUGIN_BRUSH("Icons/EditPivot_40x",        Icon20x20));      
