ModelingMode: unregister UI commands to avoid ensures in debug build
#rnx

diff --git a/Engine/Plugins/Experimental/ModelingToolsEditorMode/Source/ModelingToolsEditorMode/Private/ModelingToolsActions.cpp b/Engine/Plugins/Experimental/ModelingToolsEditorMode/Source/ModelingToolsEditorMode/Private/ModelingToolsActions.cpp
--- a/Engine/Plugins/Experimental/ModelingToolsEditorMode/Source/ModelingToolsEditorMode/Private/ModelingToolsActions.cpp
+++ b/Engine/Plugins/Experimental/ModelingToolsEditorMode/Source/ModelingToolsEditorMode/Private/ModelingToolsActions.cpp
@@ -44,6 +44,11 @@ void FModelingModeActionCommands::RegisterCommandBindings(TSharedPtr<FUICommandL
 		FExecuteAction::CreateLambda([OnCommandExecuted]() { OnCommandExecuted(EModelingModeActionCommands::FocusViewToCursor); }));
 }
 
+void FModelingModeActionCommands::UnRegisterCommandBindings(TSharedPtr<FUICommandList> UICommandList)
+{
+	const FModelingModeActionCommands& Commands = FModelingModeActionCommands::Get();
+	UICommandList->UnmapAction(Commands.FocusViewCommand);
+}
 
 
 
diff --git a/Engine/Plugins/Experimental/ModelingToolsEditorMode/Source/ModelingToolsEditorMode/Private/ModelingToolsEditorMode.cpp b/Engine/Plugins/Experimental/ModelingToolsEditorMode/Source/ModelingToolsEditorMode/Private/ModelingToolsEditorMode.cpp
--- a/Engine/Plugins/Experimental/ModelingToolsEditorMode/Source/ModelingToolsEditorMode/Private/ModelingToolsEditorMode.cpp
+++ b/Engine/Plugins/Experimental/ModelingToolsEditorMode/Source/ModelingToolsEditorMode/Private/ModelingToolsEditorMode.cpp
@@ -675,10 +675,18 @@ void FModelingToolsEditorMode::Exit()
 
 	if (Toolkit.IsValid())
 	{
+		const FModelingToolsManagerCommands& ToolManagerCommands = FModelingToolsManagerCommands::Get();
+		const TSharedRef<FUICommandList>& ToolkitCommandList = Toolkit->GetToolkitCommands();
+		ToolkitCommandList->UnmapAction(ToolManagerCommands.AcceptActiveTool);
+		ToolkitCommandList->UnmapAction(ToolManagerCommands.CancelActiveTool);
+		ToolkitCommandList->UnmapAction(ToolManagerCommands.CompleteActiveTool);
+
 		FToolkitManager::Get().CloseToolkit(Toolkit.ToSharedRef());
 		Toolkit.Reset();
 	}
 
+	FModelingModeActionCommands::UnRegisterCommandBindings(UICommandList);
+
 	// clear realtime viewport override
 	ConfigureRealTimeViewportsOverride(false);
 
diff --git a/Engine/Plugins/Experimental/ModelingToolsEditorMode/Source/ModelingToolsEditorMode/Public/ModelingToolsActions.h b/Engine/Plugins/Experimental/ModelingToolsEditorMode/Source/ModelingToolsEditorMode/Public/ModelingToolsActions.h
--- a/Engine/Plugins/Experimental/ModelingToolsEditorMode/Source/ModelingToolsEditorMode/Public/ModelingToolsActions.h
+++ b/Engine/Plugins/Experimental/ModelingToolsEditorMode/Source/ModelingToolsEditorMode/Public/ModelingToolsActions.h
@@ -23,7 +23,7 @@ public:
 	TSharedPtr<FUICommandInfo> FocusViewCommand;
 
 	static void RegisterCommandBindings(TSharedPtr<FUICommandList> UICommandList, TFunction<void(EModelingModeActionCommands)> OnCommandExecuted);
-
+	static void UnRegisterCommandBindings(TSharedPtr<FUICommandList> UICommandList);
 };
 
 
