The EditorToolsContext is now owned by the mode manager.

Updated Mesh Paint, geometry, arch gen, hair lab, and sample editor modes to reflect changes to the tools context ownership.

diff --git a/Engine/Plugins/Experimental/MeshModelingToolset/Source/MeshModelingTools/Private/AddPatchTool.cpp b/Engine/Plugins/Experimental/MeshModelingToolset/Source/MeshModelingTools/Private/AddPatchTool.cpp
--- a/Engine/Plugins/Experimental/MeshModelingToolset/Source/MeshModelingTools/Private/AddPatchTool.cpp
+++ b/Engine/Plugins/Experimental/MeshModelingToolset/Source/MeshModelingTools/Private/AddPatchTool.cpp
@@ -52,7 +52,7 @@ void UAddPatchTool::SetWorld(UWorld* World)
 	this->TargetWorld = World;
 }
 
-void UAddPatchTool::SetAssetAPI(IToolsContextAssetAPI* AssetAPIIn)
+void UAddPatchTool::SetAssetAPI(IAssetGenerationAPI* AssetAPIIn)
 {
 	this->AssetAPI = AssetAPIIn;
 }
diff --git a/Engine/Plugins/Experimental/MeshModelingToolset/Source/MeshModelingTools/Private/AddPrimitiveTool.cpp b/Engine/Plugins/Experimental/MeshModelingToolset/Source/MeshModelingTools/Private/AddPrimitiveTool.cpp
--- a/Engine/Plugins/Experimental/MeshModelingToolset/Source/MeshModelingTools/Private/AddPrimitiveTool.cpp
+++ b/Engine/Plugins/Experimental/MeshModelingToolset/Source/MeshModelingTools/Private/AddPrimitiveTool.cpp
@@ -111,7 +111,7 @@ void UAddPrimitiveTool::SetWorld(UWorld* World)
 	this->TargetWorld = World;
 }
 
-void UAddPrimitiveTool::SetAssetAPI(IToolsContextAssetAPI* AssetAPIIn)
+void UAddPrimitiveTool::SetAssetAPI(IAssetGenerationAPI* AssetAPIIn)
 {
 	this->AssetAPI = AssetAPIIn;
 }
diff --git a/Engine/Plugins/Experimental/MeshModelingToolset/Source/MeshModelingTools/Private/BakeMeshAttributeMapsTool.cpp b/Engine/Plugins/Experimental/MeshModelingToolset/Source/MeshModelingTools/Private/BakeMeshAttributeMapsTool.cpp
--- a/Engine/Plugins/Experimental/MeshModelingToolset/Source/MeshModelingTools/Private/BakeMeshAttributeMapsTool.cpp
+++ b/Engine/Plugins/Experimental/MeshModelingToolset/Source/MeshModelingTools/Private/BakeMeshAttributeMapsTool.cpp
@@ -82,7 +82,7 @@ UBakeMeshAttributeMapsTool::UBakeMeshAttributeMapsTool()
 }
 
 
-void UBakeMeshAttributeMapsTool::SetAssetAPI(IToolsContextAssetAPI* AssetAPIIn)
+void UBakeMeshAttributeMapsTool::SetAssetAPI(IAssetGenerationAPI* AssetAPIIn)
 {
 	AssetAPI = AssetAPIIn;
 }
diff --git a/Engine/Plugins/Experimental/MeshModelingToolset/Source/MeshModelingTools/Private/CombineMeshesTool.cpp b/Engine/Plugins/Experimental/MeshModelingToolset/Source/MeshModelingTools/Private/CombineMeshesTool.cpp
--- a/Engine/Plugins/Experimental/MeshModelingToolset/Source/MeshModelingTools/Private/CombineMeshesTool.cpp
+++ b/Engine/Plugins/Experimental/MeshModelingToolset/Source/MeshModelingTools/Private/CombineMeshesTool.cpp
@@ -121,7 +121,7 @@ void UCombineMeshesTool::Shutdown(EToolShutdownType ShutdownType)
 	}
 }
 
-void UCombineMeshesTool::SetAssetAPI(IToolsContextAssetAPI* AssetAPIIn)
+void UCombineMeshesTool::SetAssetAPI(IAssetGenerationAPI* AssetAPIIn)
 {
 	this->AssetAPI = AssetAPIIn;
 }
diff --git a/Engine/Plugins/Experimental/MeshModelingToolset/Source/MeshModelingTools/Private/DrawPolyPathTool.cpp b/Engine/Plugins/Experimental/MeshModelingToolset/Source/MeshModelingTools/Private/DrawPolyPathTool.cpp
--- a/Engine/Plugins/Experimental/MeshModelingToolset/Source/MeshModelingTools/Private/DrawPolyPathTool.cpp
+++ b/Engine/Plugins/Experimental/MeshModelingToolset/Source/MeshModelingTools/Private/DrawPolyPathTool.cpp
@@ -167,7 +167,7 @@ void UDrawPolyPathTool::SetWorld(UWorld* World)
 	this->TargetWorld = World;
 }
 
-void UDrawPolyPathTool::SetAssetAPI(IToolsContextAssetAPI* AssetAPIIn)
+void UDrawPolyPathTool::SetAssetAPI(IAssetGenerationAPI* AssetAPIIn)
 {
 	this->AssetAPI = AssetAPIIn;
 }
diff --git a/Engine/Plugins/Experimental/MeshModelingToolset/Source/MeshModelingTools/Private/DrawPolygonTool.cpp b/Engine/Plugins/Experimental/MeshModelingToolset/Source/MeshModelingTools/Private/DrawPolygonTool.cpp
--- a/Engine/Plugins/Experimental/MeshModelingToolset/Source/MeshModelingTools/Private/DrawPolygonTool.cpp
+++ b/Engine/Plugins/Experimental/MeshModelingToolset/Source/MeshModelingTools/Private/DrawPolygonTool.cpp
@@ -80,7 +80,7 @@ void UDrawPolygonTool::SetWorld(UWorld* World)
 	this->TargetWorld = World;
 }
 
-void UDrawPolygonTool::SetAssetAPI(IToolsContextAssetAPI* AssetAPIIn)
+void UDrawPolygonTool::SetAssetAPI(IAssetGenerationAPI* AssetAPIIn)
 {
 	this->AssetAPI = AssetAPIIn;
 }
diff --git a/Engine/Plugins/Experimental/MeshModelingToolset/Source/MeshModelingTools/Private/EditNormalsTool.cpp b/Engine/Plugins/Experimental/MeshModelingToolset/Source/MeshModelingTools/Private/EditNormalsTool.cpp
--- a/Engine/Plugins/Experimental/MeshModelingToolset/Source/MeshModelingTools/Private/EditNormalsTool.cpp
+++ b/Engine/Plugins/Experimental/MeshModelingToolset/Source/MeshModelingTools/Private/EditNormalsTool.cpp
@@ -179,7 +179,7 @@ void UEditNormalsTool::Shutdown(EToolShutdownType ShutdownType)
 	}
 }
 
-void UEditNormalsTool::SetAssetAPI(IToolsContextAssetAPI* AssetAPIIn)
+void UEditNormalsTool::SetAssetAPI(IAssetGenerationAPI* AssetAPIIn)
 {
 	this->AssetAPI = AssetAPIIn;
 }
diff --git a/Engine/Plugins/Experimental/MeshModelingToolset/Source/MeshModelingTools/Private/MeshSelectionTool.cpp b/Engine/Plugins/Experimental/MeshModelingToolset/Source/MeshModelingTools/Private/MeshSelectionTool.cpp
--- a/Engine/Plugins/Experimental/MeshModelingToolset/Source/MeshModelingTools/Private/MeshSelectionTool.cpp
+++ b/Engine/Plugins/Experimental/MeshModelingToolset/Source/MeshModelingTools/Private/MeshSelectionTool.cpp
@@ -55,7 +55,7 @@ void UMeshSelectionTool::SetWorld(UWorld* World)
 	this->TargetWorld = World;
 }
 
-void UMeshSelectionTool::SetAssetAPI(IToolsContextAssetAPI* AssetAPIIn)
+void UMeshSelectionTool::SetAssetAPI(IAssetGenerationAPI* AssetAPIIn)
 {
 	this->AssetAPI = AssetAPIIn;
 }
diff --git a/Engine/Plugins/Experimental/MeshModelingToolset/Source/MeshModelingTools/Private/MirrorTool.cpp b/Engine/Plugins/Experimental/MeshModelingToolset/Source/MeshModelingTools/Private/MirrorTool.cpp
--- a/Engine/Plugins/Experimental/MeshModelingToolset/Source/MeshModelingTools/Private/MirrorTool.cpp
+++ b/Engine/Plugins/Experimental/MeshModelingToolset/Source/MeshModelingTools/Private/MirrorTool.cpp
@@ -125,7 +125,7 @@ void UMirrorTool::SetWorld(UWorld* World)
 	TargetWorld = World;
 }
 
-void UMirrorTool::SetAssetAPI(IToolsContextAssetAPI* NewAssetApi)
+void UMirrorTool::SetAssetAPI(IAssetGenerationAPI* NewAssetApi)
 {
 	AssetAPI = NewAssetApi;
 }
diff --git a/Engine/Plugins/Experimental/MeshModelingToolset/Source/MeshModelingTools/Private/PlaneCutTool.cpp b/Engine/Plugins/Experimental/MeshModelingToolset/Source/MeshModelingTools/Private/PlaneCutTool.cpp
--- a/Engine/Plugins/Experimental/MeshModelingToolset/Source/MeshModelingTools/Private/PlaneCutTool.cpp
+++ b/Engine/Plugins/Experimental/MeshModelingToolset/Source/MeshModelingTools/Private/PlaneCutTool.cpp
@@ -319,7 +319,7 @@ void UPlaneCutTool::Shutdown(EToolShutdownType ShutdownType)
 	GizmoManager->DestroyAllGizmosByOwner(this);
 }
 
-void UPlaneCutTool::SetAssetAPI(IToolsContextAssetAPI* AssetAPIIn)
+void UPlaneCutTool::SetAssetAPI(IAssetGenerationAPI* AssetAPIIn)
 {
 	this->AssetAPI = AssetAPIIn;
 }
diff --git a/Engine/Plugins/Experimental/MeshModelingToolset/Source/MeshModelingTools/Private/RemoveOccludedTrianglesTool.cpp b/Engine/Plugins/Experimental/MeshModelingToolset/Source/MeshModelingTools/Private/RemoveOccludedTrianglesTool.cpp
--- a/Engine/Plugins/Experimental/MeshModelingToolset/Source/MeshModelingTools/Private/RemoveOccludedTrianglesTool.cpp
+++ b/Engine/Plugins/Experimental/MeshModelingToolset/Source/MeshModelingTools/Private/RemoveOccludedTrianglesTool.cpp
@@ -265,7 +265,7 @@ void URemoveOccludedTrianglesTool::Shutdown(EToolShutdownType ShutdownType)
 	}
 }
 
-void URemoveOccludedTrianglesTool::SetAssetAPI(IToolsContextAssetAPI* AssetAPIIn)
+void URemoveOccludedTrianglesTool::SetAssetAPI(IAssetGenerationAPI* AssetAPIIn)
 {
 	this->AssetAPI = AssetAPIIn;
 }
diff --git a/Engine/Plugins/Experimental/MeshModelingToolset/Source/MeshModelingTools/Private/UVProjectionTool.cpp b/Engine/Plugins/Experimental/MeshModelingToolset/Source/MeshModelingTools/Private/UVProjectionTool.cpp
--- a/Engine/Plugins/Experimental/MeshModelingToolset/Source/MeshModelingTools/Private/UVProjectionTool.cpp
+++ b/Engine/Plugins/Experimental/MeshModelingToolset/Source/MeshModelingTools/Private/UVProjectionTool.cpp
@@ -216,7 +216,7 @@ void UUVProjectionTool::Shutdown(EToolShutdownType ShutdownType)
 	TransformGizmos.Empty();
 }
 
-void UUVProjectionTool::SetAssetAPI(IToolsContextAssetAPI* AssetAPIIn)
+void UUVProjectionTool::SetAssetAPI(IAssetGenerationAPI* AssetAPIIn)
 {
 	this->AssetAPI = AssetAPIIn;
 }
diff --git a/Engine/Plugins/Experimental/MeshModelingToolset/Source/MeshModelingTools/Public/AddPatchTool.h b/Engine/Plugins/Experimental/MeshModelingToolset/Source/MeshModelingTools/Public/AddPatchTool.h
--- a/Engine/Plugins/Experimental/MeshModelingToolset/Source/MeshModelingTools/Public/AddPatchTool.h
+++ b/Engine/Plugins/Experimental/MeshModelingToolset/Source/MeshModelingTools/Public/AddPatchTool.h
@@ -11,6 +11,8 @@
 #include "Properties/MeshMaterialProperties.h"
 #include "AddPatchTool.generated.h"
 
+class IAssetGenerationAPI;
+
 class FDynamicMesh3;
 
 /**
@@ -22,7 +24,7 @@ class MESHMODELINGTOOLS_API UAddPatchToolBuilder : public UInteractiveToolBuilde
 	GENERATED_BODY()
 
 public:
-	IToolsContextAssetAPI* AssetAPI;
+	IAssetGenerationAPI* AssetAPI;
 
 	UAddPatchToolBuilder() 
 	{
@@ -76,7 +78,7 @@ class MESHMODELINGTOOLS_API UAddPatchTool : public USingleClickTool, public IHov
 
 public:
 	virtual void SetWorld(UWorld* World);
-	virtual void SetAssetAPI(IToolsContextAssetAPI* AssetAPI);
+	virtual void SetAssetAPI(IAssetGenerationAPI* AssetAPI);
 
 	virtual void Setup() override;
 	virtual void Shutdown(EToolShutdownType ShutdownType) override;
@@ -114,7 +116,7 @@ protected:
 
 protected:
 	UWorld* TargetWorld;
-	IToolsContextAssetAPI* AssetAPI;
+	IAssetGenerationAPI* AssetAPI;
 
 	FBox WorldBounds;
 
diff --git a/Engine/Plugins/Experimental/MeshModelingToolset/Source/MeshModelingTools/Public/AddPrimitiveTool.h b/Engine/Plugins/Experimental/MeshModelingToolset/Source/MeshModelingTools/Public/AddPrimitiveTool.h
--- a/Engine/Plugins/Experimental/MeshModelingToolset/Source/MeshModelingTools/Public/AddPrimitiveTool.h
+++ b/Engine/Plugins/Experimental/MeshModelingToolset/Source/MeshModelingTools/Public/AddPrimitiveTool.h
@@ -14,6 +14,7 @@
 #include "AddPrimitiveTool.generated.h"
 
 class FDynamicMesh3;
+class IAssetGenerationAPI;
 
 /**
  * Builder
@@ -39,7 +40,7 @@ public:
 		Sphere
 	};
 
-	IToolsContextAssetAPI* AssetAPI;
+	IAssetGenerationAPI* AssetAPI;
 
 	UAddPrimitiveToolBuilder()
 	{
@@ -376,7 +377,7 @@ public:
 	UAddPrimitiveTool(const FObjectInitializer&);
 
 	virtual void SetWorld(UWorld* World);
-	virtual void SetAssetAPI(IToolsContextAssetAPI* AssetAPI);
+	virtual void SetAssetAPI(IAssetGenerationAPI* AssetAPI);
 
 	virtual void Setup() override;
 	virtual void Shutdown(EToolShutdownType ShutdownType) override;
@@ -436,7 +437,7 @@ protected:
 
 protected:
 	UWorld* TargetWorld;
-	IToolsContextAssetAPI* AssetAPI;
+	IAssetGenerationAPI* AssetAPI;
 
 	void UpdatePreviewPosition(const FInputDeviceRay& ClickPos);
 	FFrame3f ShapeFrame;
diff --git a/Engine/Plugins/Experimental/MeshModelingToolset/Source/MeshModelingTools/Public/BakeMeshAttributeMapsTool.h b/Engine/Plugins/Experimental/MeshModelingToolset/Source/MeshModelingTools/Public/BakeMeshAttributeMapsTool.h
--- a/Engine/Plugins/Experimental/MeshModelingToolset/Source/MeshModelingTools/Public/BakeMeshAttributeMapsTool.h
+++ b/Engine/Plugins/Experimental/MeshModelingToolset/Source/MeshModelingTools/Public/BakeMeshAttributeMapsTool.h
@@ -17,6 +17,7 @@ struct FMeshDescription;
 class USimpleDynamicMeshComponent;
 class UMaterialInstanceDynamic;
 class UTexture2D;
+class IAssetGenerationAPI;
 template<typename RealType> class TMeshTangents;
 class FMeshImageBakingCache;
 
@@ -29,7 +30,7 @@ class MESHMODELINGTOOLS_API UBakeMeshAttributeMapsToolBuilder : public UInteract
 	GENERATED_BODY()
 
 public:
-	IToolsContextAssetAPI* AssetAPI = nullptr;
+	IAssetGenerationAPI* AssetAPI = nullptr;
 
 	virtual bool CanBuildTool(const FToolBuilderState& SceneState) const override;
 	virtual UInteractiveTool* BuildTool(const FToolBuilderState& SceneState) const override;
@@ -234,7 +235,7 @@ class MESHMODELINGTOOLS_API UBakeMeshAttributeMapsTool : public UMultiSelectionT
 public:
 	UBakeMeshAttributeMapsTool();
 
-	virtual void SetAssetAPI(IToolsContextAssetAPI* AssetAPI);
+	virtual void SetAssetAPI(IAssetGenerationAPI* AssetAPI);
 
 	virtual void Setup() override;
 	virtual void Shutdown(EToolShutdownType ShutdownType) override;
@@ -270,7 +271,7 @@ protected:
 
 
 protected:
-	IToolsContextAssetAPI* AssetAPI = nullptr;
+	IAssetGenerationAPI* AssetAPI = nullptr;
 
 	USimpleDynamicMeshComponent* DynamicMeshComponent;
 
diff --git a/Engine/Plugins/Experimental/MeshModelingToolset/Source/MeshModelingTools/Public/CombineMeshesTool.h b/Engine/Plugins/Experimental/MeshModelingToolset/Source/MeshModelingTools/Public/CombineMeshesTool.h
--- a/Engine/Plugins/Experimental/MeshModelingToolset/Source/MeshModelingTools/Public/CombineMeshesTool.h
+++ b/Engine/Plugins/Experimental/MeshModelingToolset/Source/MeshModelingTools/Public/CombineMeshesTool.h
@@ -15,7 +15,7 @@
 // predeclarations
 struct FMeshDescription;
 class USimpleDynamicMeshComponent;
-
+class IAssetGenerationAPI;
 
 
 
@@ -29,7 +29,7 @@ class MESHMODELINGTOOLS_API UCombineMeshesToolBuilder : public UInteractiveToolB
 	GENERATED_BODY()
 
 public:
-	IToolsContextAssetAPI* AssetAPI = nullptr;
+	IAssetGenerationAPI* AssetAPI = nullptr;
 	bool bIsDuplicateTool = false;
 
 	virtual bool CanBuildTool(const FToolBuilderState& SceneState) const override;
@@ -74,7 +74,7 @@ public:
 	virtual void Shutdown(EToolShutdownType ShutdownType) override;
 
 	virtual void SetWorld(UWorld* World);
-	virtual void SetAssetAPI(IToolsContextAssetAPI* AssetAPI);
+	virtual void SetAssetAPI(IAssetGenerationAPI* AssetAPI);
 
 	virtual bool HasCancel() const override { return true; }
 	virtual bool HasAccept() const override;
@@ -87,7 +87,7 @@ protected:
 
 protected:
 	UWorld* TargetWorld;
-	IToolsContextAssetAPI* AssetAPI;
+	IAssetGenerationAPI* AssetAPI;
 
 	bool bDuplicateMode;
 
diff --git a/Engine/Plugins/Experimental/MeshModelingToolset/Source/MeshModelingTools/Public/DrawAndRevolveTool.h b/Engine/Plugins/Experimental/MeshModelingToolset/Source/MeshModelingTools/Public/DrawAndRevolveTool.h
--- a/Engine/Plugins/Experimental/MeshModelingToolset/Source/MeshModelingTools/Public/DrawAndRevolveTool.h
+++ b/Engine/Plugins/Experimental/MeshModelingToolset/Source/MeshModelingTools/Public/DrawAndRevolveTool.h
@@ -19,6 +19,7 @@ class UCollectSurfacePathMechanic;
 class UConstructionPlaneMechanic;
 class UCurveControlPointsMechanic;
 class FCurveSweepOp;
+class IAssetGenerationAPI;
 
 UCLASS()
 class MESHMODELINGTOOLS_API UDrawAndRevolveToolBuilder : public UInteractiveToolBuilder
@@ -26,7 +27,7 @@ class MESHMODELINGTOOLS_API UDrawAndRevolveToolBuilder : public UInteractiveTool
 	GENERATED_BODY()
 
 public:
-	IToolsContextAssetAPI* AssetAPI = nullptr;
+	IAssetGenerationAPI* AssetAPI = nullptr;
 
 	virtual bool CanBuildTool(const FToolBuilderState& SceneState) const override;
 	virtual UInteractiveTool* BuildTool(const FToolBuilderState& SceneState) const override;
@@ -83,7 +84,7 @@ class MESHMODELINGTOOLS_API UDrawAndRevolveTool : public UInteractiveTool
 
 public:
 	virtual void SetWorld(UWorld* World) { TargetWorld = World; }
-	virtual void SetAssetAPI(IToolsContextAssetAPI* NewAssetApi) { AssetAPI = NewAssetApi; }
+	virtual void SetAssetAPI(IAssetGenerationAPI* NewAssetApi) { AssetAPI = NewAssetApi; }
 
 	virtual void RegisterActions(FInteractiveToolActionSet& ActionSet) override;
 	void OnBackspacePress();
@@ -103,7 +104,7 @@ public:
 protected:
 
 	UWorld* TargetWorld;
-	IToolsContextAssetAPI* AssetAPI;
+	IAssetGenerationAPI* AssetAPI;
 
 	FViewCameraState CameraState;
 
diff --git a/Engine/Plugins/Experimental/MeshModelingToolset/Source/MeshModelingTools/Public/DrawPolyPathTool.h b/Engine/Plugins/Experimental/MeshModelingToolset/Source/MeshModelingTools/Public/DrawPolyPathTool.h
--- a/Engine/Plugins/Experimental/MeshModelingToolset/Source/MeshModelingTools/Public/DrawPolyPathTool.h
+++ b/Engine/Plugins/Experimental/MeshModelingToolset/Source/MeshModelingTools/Public/DrawPolyPathTool.h
@@ -16,6 +16,7 @@
 class FMeshVertexChangeBuilder;
 class UTransformGizmo;
 class UTransformProxy;
+class IAssetGenerationAPI;
 
 /**
  * ToolBuilder
@@ -25,7 +26,7 @@ class MESHMODELINGTOOLS_API UDrawPolyPathToolBuilder : public UMeshSurfacePointT
 {
 	GENERATED_BODY()
 public:
-	IToolsContextAssetAPI* AssetAPI = nullptr;
+	IAssetGenerationAPI* AssetAPI = nullptr;
 
 	virtual bool CanBuildTool(const FToolBuilderState& SceneState) const override;
 
@@ -129,7 +130,7 @@ public:
 	virtual void RegisterActions(FInteractiveToolActionSet& ActionSet) override;
 
 	virtual void SetWorld(UWorld* World);
-	virtual void SetAssetAPI(IToolsContextAssetAPI* AssetAPI);
+	virtual void SetAssetAPI(IAssetGenerationAPI* AssetAPI);
 
 	virtual void Setup() override;
 	virtual void Shutdown(EToolShutdownType ShutdownType) override;
@@ -155,7 +156,7 @@ public:
 
 protected:
 	UWorld* TargetWorld;
-	IToolsContextAssetAPI* AssetAPI;
+	IAssetGenerationAPI* AssetAPI;
 
 	UPROPERTY()
 	UDrawPolyPathProperties* TransformProps;
diff --git a/Engine/Plugins/Experimental/MeshModelingToolset/Source/MeshModelingTools/Public/DrawPolygonTool.h b/Engine/Plugins/Experimental/MeshModelingToolset/Source/MeshModelingTools/Public/DrawPolygonTool.h
--- a/Engine/Plugins/Experimental/MeshModelingToolset/Source/MeshModelingTools/Public/DrawPolygonTool.h
+++ b/Engine/Plugins/Experimental/MeshModelingToolset/Source/MeshModelingTools/Public/DrawPolygonTool.h
@@ -18,6 +18,7 @@
 
 class UTransformGizmo;
 class UTransformProxy;
+class IAssetGenerationAPI;
 
 /**
  *
@@ -28,7 +29,7 @@ class MESHMODELINGTOOLS_API UDrawPolygonToolBuilder : public UInteractiveToolBui
 	GENERATED_BODY()
 
 public:
-	IToolsContextAssetAPI* AssetAPI;
+	IAssetGenerationAPI* AssetAPI;
 
 	UDrawPolygonToolBuilder() 
 	{
@@ -171,7 +172,7 @@ public:
 	virtual void RegisterActions(FInteractiveToolActionSet& ActionSet) override;
 
 	virtual void SetWorld(UWorld* World);
-	virtual void SetAssetAPI(IToolsContextAssetAPI* AssetAPI);
+	virtual void SetAssetAPI(IAssetGenerationAPI* AssetAPI);
 
 	virtual void Setup() override;
 	virtual void Shutdown(EToolShutdownType ShutdownType) override;
@@ -244,7 +245,7 @@ protected:
 
 protected:
 	UWorld* TargetWorld;
-	IToolsContextAssetAPI* AssetAPI;
+	IAssetGenerationAPI* AssetAPI;
 
 	FViewCameraState CameraState;
 
diff --git a/Engine/Plugins/Experimental/MeshModelingToolset/Source/MeshModelingTools/Public/EditNormalsTool.h b/Engine/Plugins/Experimental/MeshModelingToolset/Source/MeshModelingTools/Public/EditNormalsTool.h
--- a/Engine/Plugins/Experimental/MeshModelingToolset/Source/MeshModelingTools/Public/EditNormalsTool.h
+++ b/Engine/Plugins/Experimental/MeshModelingToolset/Source/MeshModelingTools/Public/EditNormalsTool.h
@@ -16,7 +16,7 @@
 // predeclarations
 struct FMeshDescription;
 class USimpleDynamicMeshComponent;
-
+class IAssetGenerationAPI;
 
 
 
@@ -30,7 +30,7 @@ class MESHMODELINGTOOLS_API UEditNormalsToolBuilder : public UInteractiveToolBui
 	GENERATED_BODY()
 
 public:
-	IToolsContextAssetAPI* AssetAPI = nullptr;
+	IAssetGenerationAPI* AssetAPI = nullptr;
 
 	virtual bool CanBuildTool(const FToolBuilderState& SceneState) const override;
 	virtual UInteractiveTool* BuildTool(const FToolBuilderState& SceneState) const override;
@@ -139,7 +139,7 @@ public:
 	virtual void Shutdown(EToolShutdownType ShutdownType) override;
 
 	virtual void SetWorld(UWorld* World);
-	virtual void SetAssetAPI(IToolsContextAssetAPI* AssetAPI);
+	virtual void SetAssetAPI(IAssetGenerationAPI* AssetAPI);
 
 	virtual void OnTick(float DeltaTime) override;
 	virtual void Render(IToolsContextRenderAPI* RenderAPI) override;
@@ -171,7 +171,7 @@ protected:
 	TArray<TSharedPtr<FDynamicMesh3>> OriginalDynamicMeshes;
 
 	UWorld* TargetWorld;
-	IToolsContextAssetAPI* AssetAPI;
+	IAssetGenerationAPI* AssetAPI;
 
 	FViewCameraState CameraState;
 
diff --git a/Engine/Plugins/Experimental/MeshModelingToolset/Source/MeshModelingTools/Public/MeshSelectionTool.h b/Engine/Plugins/Experimental/MeshModelingToolset/Source/MeshModelingTools/Public/MeshSelectionTool.h
--- a/Engine/Plugins/Experimental/MeshModelingToolset/Source/MeshModelingTools/Public/MeshSelectionTool.h
+++ b/Engine/Plugins/Experimental/MeshModelingToolset/Source/MeshModelingTools/Public/MeshSelectionTool.h
@@ -9,6 +9,8 @@
 #include "DynamicMeshOctree3.h"
 #include "MeshSelectionTool.generated.h"
 
+class IAssetGenerationAPI;
+
 /**
  *
  */
@@ -18,7 +20,7 @@ class MESHMODELINGTOOLS_API UMeshSelectionToolBuilder : public UMeshSurfacePoint
 	GENERATED_BODY()
 
 public:
-	IToolsContextAssetAPI* AssetAPI = nullptr;
+	IAssetGenerationAPI* AssetAPI = nullptr;
 
 	virtual UMeshSurfacePointTool* CreateNewTool(const FToolBuilderState& SceneState) const override;
 };
@@ -265,7 +267,7 @@ public:
 	UMeshSelectionTool();
 
 	virtual void SetWorld(UWorld* World);
-	virtual void SetAssetAPI(IToolsContextAssetAPI* AssetAPI);
+	virtual void SetAssetAPI(IAssetGenerationAPI* AssetAPI);
 
 	virtual void RegisterActions(FInteractiveToolActionSet& ActionSet) override;
 
@@ -317,7 +319,7 @@ protected:
 	TArray<AActor*> SpawnedActors;
 
 	UWorld* TargetWorld;
-	IToolsContextAssetAPI* AssetAPI;
+	IAssetGenerationAPI* AssetAPI;
 
 	// note: ideally this octree would be part of PreviewMesh!
 	TUniquePtr<FDynamicMeshOctree3> Octree;
diff --git a/Engine/Plugins/Experimental/MeshModelingToolset/Source/MeshModelingTools/Public/MirrorTool.h b/Engine/Plugins/Experimental/MeshModelingToolset/Source/MeshModelingTools/Public/MirrorTool.h
--- a/Engine/Plugins/Experimental/MeshModelingToolset/Source/MeshModelingTools/Public/MirrorTool.h
+++ b/Engine/Plugins/Experimental/MeshModelingToolset/Source/MeshModelingTools/Public/MirrorTool.h
@@ -17,6 +17,7 @@
 
 class UTransformGizmo;
 class UTransformProxy;
+class IAssetGenerationAPI;
 
 UCLASS()
 class MESHMODELINGTOOLS_API UMirrorToolBuilder : public UInteractiveToolBuilder
@@ -24,7 +25,7 @@ class MESHMODELINGTOOLS_API UMirrorToolBuilder : public UInteractiveToolBuilder
 	GENERATED_BODY()
 
 public:
-	IToolsContextAssetAPI* AssetAPI = nullptr;
+	IAssetGenerationAPI* AssetAPI = nullptr;
 
 	virtual bool CanBuildTool(const FToolBuilderState& SceneState) const override;
 	virtual UInteractiveTool* BuildTool(const FToolBuilderState& SceneState) const override;
@@ -195,7 +196,7 @@ public:
 	virtual void Shutdown(EToolShutdownType ShutdownType) override;
 
 	virtual void SetWorld(UWorld* World);
-	virtual void SetAssetAPI(IToolsContextAssetAPI* AssetAPI);
+	virtual void SetAssetAPI(IAssetGenerationAPI* AssetAPI);
 
 	virtual void OnTick(float DeltaTime) override;
 	virtual void Render(IToolsContextRenderAPI* RenderAPI) override;
@@ -228,7 +229,7 @@ protected:
 	FVector3d MirrorPlaneNormal = FVector3d::UnitZ();
 
 	UWorld* TargetWorld;
-	IToolsContextAssetAPI* AssetAPI;
+	IAssetGenerationAPI* AssetAPI;
 
 	// flags used to identify modifier keys/buttons
 	static const int SnappingToggleModifierId = 1;
diff --git a/Engine/Plugins/Experimental/MeshModelingToolset/Source/MeshModelingTools/Public/PlaneCutTool.h b/Engine/Plugins/Experimental/MeshModelingToolset/Source/MeshModelingTools/Public/PlaneCutTool.h
--- a/Engine/Plugins/Experimental/MeshModelingToolset/Source/MeshModelingTools/Public/PlaneCutTool.h
+++ b/Engine/Plugins/Experimental/MeshModelingToolset/Source/MeshModelingTools/Public/PlaneCutTool.h
@@ -17,7 +17,7 @@
 // predeclarations
 class UTransformGizmo;
 class UTransformProxy;
-
+class IAssetGenerationAPI;
 
 
 
@@ -31,7 +31,7 @@ class MESHMODELINGTOOLS_API UPlaneCutToolBuilder : public UInteractiveToolBuilde
 	GENERATED_BODY()
 
 public:
-	IToolsContextAssetAPI* AssetAPI = nullptr;
+	IAssetGenerationAPI* AssetAPI = nullptr;
 
 	virtual bool CanBuildTool(const FToolBuilderState& SceneState) const override;
 	virtual UInteractiveTool* BuildTool(const FToolBuilderState& SceneState) const override;
@@ -129,7 +129,7 @@ public:
 	virtual void Shutdown(EToolShutdownType ShutdownType) override;
 
 	virtual void SetWorld(UWorld* World);
-	virtual void SetAssetAPI(IToolsContextAssetAPI* AssetAPI);
+	virtual void SetAssetAPI(IAssetGenerationAPI* AssetAPI);
 
 	virtual void RegisterActions(FInteractiveToolActionSet& ActionSet) override;
 
@@ -184,7 +184,7 @@ protected:
 	TArray<float> MeshUVScaleFactor;
 
 	UWorld* TargetWorld;
-	IToolsContextAssetAPI* AssetAPI;
+	IAssetGenerationAPI* AssetAPI;
 
 	FViewCameraState CameraState;
 
diff --git a/Engine/Plugins/Experimental/MeshModelingToolset/Source/MeshModelingTools/Public/RemoveOccludedTrianglesTool.h b/Engine/Plugins/Experimental/MeshModelingToolset/Source/MeshModelingTools/Public/RemoveOccludedTrianglesTool.h
--- a/Engine/Plugins/Experimental/MeshModelingToolset/Source/MeshModelingTools/Public/RemoveOccludedTrianglesTool.h
+++ b/Engine/Plugins/Experimental/MeshModelingToolset/Source/MeshModelingTools/Public/RemoveOccludedTrianglesTool.h
@@ -13,6 +13,7 @@
 #include "BaseTools/SingleClickTool.h"
 #include "RemoveOccludedTrianglesTool.generated.h"
 
+class IAssetGenerationAPI;
 
 /**
  *
@@ -23,7 +24,7 @@ class MESHMODELINGTOOLS_API URemoveOccludedTrianglesToolBuilder : public UIntera
 	GENERATED_BODY()
 
 public:
-	IToolsContextAssetAPI* AssetAPI = nullptr;
+	IAssetGenerationAPI* AssetAPI = nullptr;
 
 	virtual bool CanBuildTool(const FToolBuilderState& SceneState) const override;
 	virtual UInteractiveTool* BuildTool(const FToolBuilderState& SceneState) const override;
@@ -144,7 +145,7 @@ public:
 	virtual void Shutdown(EToolShutdownType ShutdownType) override;
 
 	virtual void SetWorld(UWorld* World);
-	virtual void SetAssetAPI(IToolsContextAssetAPI* AssetAPI);
+	virtual void SetAssetAPI(IAssetGenerationAPI* AssetAPI);
 
 	virtual void OnTick(float DeltaTime) override;
 
@@ -186,7 +187,7 @@ protected:
 	TArray<int32> TargetToPreviewIdx;
 
 	UWorld* TargetWorld;
-	IToolsContextAssetAPI* AssetAPI;
+	IAssetGenerationAPI* AssetAPI;
 
 	FViewCameraState CameraState;
 
diff --git a/Engine/Plugins/Experimental/MeshModelingToolset/Source/MeshModelingTools/Public/RevolveBoundaryTool.h b/Engine/Plugins/Experimental/MeshModelingToolset/Source/MeshModelingTools/Public/RevolveBoundaryTool.h
--- a/Engine/Plugins/Experimental/MeshModelingToolset/Source/MeshModelingTools/Public/RevolveBoundaryTool.h
+++ b/Engine/Plugins/Experimental/MeshModelingToolset/Source/MeshModelingTools/Public/RevolveBoundaryTool.h
@@ -17,6 +17,8 @@
 
 #include "RevolveBoundaryTool.generated.h"
 
+class IAssetGenerationAPI;
+
 // Tool Builder
 
 UCLASS()
@@ -25,7 +27,7 @@ class MESHMODELINGTOOLS_API URevolveBoundaryToolBuilder : public UInteractiveToo
 	GENERATED_BODY()
 
 public:
-	IToolsContextAssetAPI* AssetAPI = nullptr;
+	IAssetGenerationAPI* AssetAPI = nullptr;
 
 	bool CanBuildTool(const FToolBuilderState& SceneState) const override;
 	UInteractiveTool* BuildTool(const FToolBuilderState& SceneState) const override;
@@ -73,7 +75,7 @@ class MESHMODELINGTOOLS_API URevolveBoundaryTool : public UMeshBoundaryToolBase
 	GENERATED_BODY()
 
 public:
-	virtual void SetAssetAPI(IToolsContextAssetAPI* NewAssetApi) { AssetAPI = NewAssetApi; }
+	virtual void SetAssetAPI(IAssetGenerationAPI* NewAssetApi) { AssetAPI = NewAssetApi; }
 
 	virtual bool CanAccept() const override;
 
@@ -91,7 +93,7 @@ protected:
 	bool bAlignAxisOnClick = false;
 	int32 AlignAxisModifier = 2;
 
-	IToolsContextAssetAPI* AssetAPI = nullptr;
+	IAssetGenerationAPI* AssetAPI = nullptr;
 
 	UPROPERTY()
 	URevolveBoundaryToolProperties* Settings = nullptr;
diff --git a/Engine/Plugins/Experimental/MeshModelingToolset/Source/MeshModelingTools/Public/UVProjectionTool.h b/Engine/Plugins/Experimental/MeshModelingToolset/Source/MeshModelingTools/Public/UVProjectionTool.h
--- a/Engine/Plugins/Experimental/MeshModelingToolset/Source/MeshModelingTools/Public/UVProjectionTool.h
+++ b/Engine/Plugins/Experimental/MeshModelingToolset/Source/MeshModelingTools/Public/UVProjectionTool.h
@@ -21,6 +21,7 @@ struct FMeshDescription;
 class USimpleDynamicMeshComponent;
 class UTransformGizmo;
 class UTransformProxy;
+class IAssetGenerationAPI;
 
 
 /**
@@ -32,7 +33,7 @@ class MESHMODELINGTOOLS_API UUVProjectionToolBuilder : public UInteractiveToolBu
 	GENERATED_BODY()
 
 public:
-	IToolsContextAssetAPI* AssetAPI = nullptr;
+	IAssetGenerationAPI* AssetAPI = nullptr;
 
 	virtual bool CanBuildTool(const FToolBuilderState& SceneState) const override;
 	virtual UInteractiveTool* BuildTool(const FToolBuilderState& SceneState) const override;
@@ -134,7 +135,7 @@ public:
 	virtual void Shutdown(EToolShutdownType ShutdownType) override;
 
 	virtual void SetWorld(UWorld* World, UInteractiveGizmoManager* GizmoManagerIn);
-	virtual void SetAssetAPI(IToolsContextAssetAPI* AssetAPI);
+	virtual void SetAssetAPI(IAssetGenerationAPI* AssetAPI);
 
 	virtual void OnTick(float DeltaTime) override;
 	virtual void Render(IToolsContextRenderAPI* RenderAPI) override;
@@ -180,7 +181,7 @@ protected:
 	TArray<FDynamicMesh3> ReferencePrimitives;
 
 	UWorld* TargetWorld;
-	IToolsContextAssetAPI* AssetAPI;
+	IAssetGenerationAPI* AssetAPI;
 	UInteractiveGizmoManager* GizmoManager;
 
 	FViewCameraState CameraState;
diff --git a/Engine/Plugins/Experimental/MeshModelingToolset/Source/MeshModelingToolsEditorOnly/Private/MergeMeshesTool.cpp b/Engine/Plugins/Experimental/MeshModelingToolset/Source/MeshModelingToolsEditorOnly/Private/MergeMeshesTool.cpp
--- a/Engine/Plugins/Experimental/MeshModelingToolset/Source/MeshModelingToolsEditorOnly/Private/MergeMeshesTool.cpp
+++ b/Engine/Plugins/Experimental/MeshModelingToolset/Source/MeshModelingToolsEditorOnly/Private/MergeMeshesTool.cpp
@@ -67,7 +67,7 @@ void UMergeMeshesTool::SetWorld(UWorld* World)
 	this->TargetWorld = World;
 }
 
-void UMergeMeshesTool::SetAssetAPI(IToolsContextAssetAPI* AssetAPIIn)
+void UMergeMeshesTool::SetAssetAPI(IAssetGenerationAPI* AssetAPIIn)
 {
 	this->AssetAPI = AssetAPIIn;
 }
diff --git a/Engine/Plugins/Experimental/MeshModelingToolset/Source/MeshModelingToolsEditorOnly/Private/ParameterizeMeshTool.cpp b/Engine/Plugins/Experimental/MeshModelingToolset/Source/MeshModelingToolsEditorOnly/Private/ParameterizeMeshTool.cpp
--- a/Engine/Plugins/Experimental/MeshModelingToolset/Source/MeshModelingToolsEditorOnly/Private/ParameterizeMeshTool.cpp
+++ b/Engine/Plugins/Experimental/MeshModelingToolset/Source/MeshModelingToolsEditorOnly/Private/ParameterizeMeshTool.cpp
@@ -57,7 +57,7 @@ void UParameterizeMeshTool::SetWorld(UWorld* World)
 	this->TargetWorld = World;
 }
 
-void UParameterizeMeshTool::SetAssetAPI(IToolsContextAssetAPI* AssetAPIIn)
+void UParameterizeMeshTool::SetAssetAPI(IAssetGenerationAPI* AssetAPIIn)
 {
 	this->AssetAPI = AssetAPIIn;
 }
diff --git a/Engine/Plugins/Experimental/MeshModelingToolset/Source/MeshModelingToolsEditorOnly/Private/ShapeSprayTool.cpp b/Engine/Plugins/Experimental/MeshModelingToolset/Source/MeshModelingToolsEditorOnly/Private/ShapeSprayTool.cpp
--- a/Engine/Plugins/Experimental/MeshModelingToolset/Source/MeshModelingToolsEditorOnly/Private/ShapeSprayTool.cpp
+++ b/Engine/Plugins/Experimental/MeshModelingToolset/Source/MeshModelingToolsEditorOnly/Private/ShapeSprayTool.cpp
@@ -50,7 +50,7 @@ void UShapeSprayTool::SetWorld(UWorld* World)
 	this->TargetWorld = World;
 }
 
-void UShapeSprayTool::SetAssetAPI(IToolsContextAssetAPI* AssetAPIIn)
+void UShapeSprayTool::SetAssetAPI(IAssetGenerationAPI* AssetAPIIn)
 {
 	this->AssetAPI = AssetAPIIn;
 }
diff --git a/Engine/Plugins/Experimental/MeshModelingToolset/Source/MeshModelingToolsEditorOnly/Private/UVLayoutTool.cpp b/Engine/Plugins/Experimental/MeshModelingToolset/Source/MeshModelingToolsEditorOnly/Private/UVLayoutTool.cpp
--- a/Engine/Plugins/Experimental/MeshModelingToolset/Source/MeshModelingToolsEditorOnly/Private/UVLayoutTool.cpp
+++ b/Engine/Plugins/Experimental/MeshModelingToolset/Source/MeshModelingToolsEditorOnly/Private/UVLayoutTool.cpp
@@ -192,7 +192,7 @@ void UUVLayoutTool::Shutdown(EToolShutdownType ShutdownType)
 	}
 }
 
-void UUVLayoutTool::SetAssetAPI(IToolsContextAssetAPI* AssetAPIIn)
+void UUVLayoutTool::SetAssetAPI(IAssetGenerationAPI* AssetAPIIn)
 {
 	this->AssetAPI = AssetAPIIn;
 }
diff --git a/Engine/Plugins/Experimental/MeshModelingToolset/Source/MeshModelingToolsEditorOnly/Private/VoxelCSGMeshesTool.cpp b/Engine/Plugins/Experimental/MeshModelingToolset/Source/MeshModelingToolsEditorOnly/Private/VoxelCSGMeshesTool.cpp
--- a/Engine/Plugins/Experimental/MeshModelingToolset/Source/MeshModelingToolsEditorOnly/Private/VoxelCSGMeshesTool.cpp
+++ b/Engine/Plugins/Experimental/MeshModelingToolset/Source/MeshModelingToolsEditorOnly/Private/VoxelCSGMeshesTool.cpp
@@ -62,7 +62,7 @@ void UVoxelCSGMeshesTool::SetWorld(UWorld* World)
 	this->TargetWorld = World;
 }
 
-void UVoxelCSGMeshesTool::SetAssetAPI(IToolsContextAssetAPI* AssetAPIIn)
+void UVoxelCSGMeshesTool::SetAssetAPI(IAssetGenerationAPI* AssetAPIIn)
 {
 	this->AssetAPI = AssetAPIIn;
 }
diff --git a/Engine/Plugins/Experimental/MeshModelingToolset/Source/MeshModelingToolsEditorOnly/Public/BspConversionTool.h b/Engine/Plugins/Experimental/MeshModelingToolset/Source/MeshModelingToolsEditorOnly/Public/BspConversionTool.h
--- a/Engine/Plugins/Experimental/MeshModelingToolset/Source/MeshModelingToolsEditorOnly/Public/BspConversionTool.h
+++ b/Engine/Plugins/Experimental/MeshModelingToolset/Source/MeshModelingToolsEditorOnly/Public/BspConversionTool.h
@@ -14,6 +14,7 @@
 
 #include "BspConversionTool.generated.h"
 
+class IAssetGenerationAPI;
 
 /**
  * Builder for UBspConversionTool.
@@ -24,7 +25,7 @@ class MESHMODELINGTOOLSEDITORONLY_API UBspConversionToolBuilder : public UIntera
 	GENERATED_BODY()
 
 public:
-	IToolsContextAssetAPI* AssetAPI = nullptr;
+	IAssetGenerationAPI* AssetAPI = nullptr;
 
 	virtual bool CanBuildTool(const FToolBuilderState& SceneState) const override;
 	virtual UInteractiveTool* BuildTool(const FToolBuilderState& SceneState) const override;
@@ -140,7 +141,7 @@ public:
 	virtual void Shutdown(EToolShutdownType ShutdownType) override;
 
 	virtual void SetWorld(UWorld* World) { this->TargetWorld = World; }
-	virtual void SetAssetAPI(IToolsContextAssetAPI* InAssetAPI) { this->AssetAPI = InAssetAPI; }
+	virtual void SetAssetAPI(IAssetGenerationAPI* InAssetAPI) { this->AssetAPI = InAssetAPI; }
 
 	virtual void OnTick(float DeltaTime) override;
 	virtual void Render(IToolsContextRenderAPI* RenderAPI) override;
@@ -163,7 +164,7 @@ protected:
 	bool bCanAccept = false;
 
 	UWorld* TargetWorld;
-	IToolsContextAssetAPI* AssetAPI;
+	IAssetGenerationAPI* AssetAPI;
 	EBspConversionToolAction PendingAction;
 
 	UPROPERTY()
diff --git a/Engine/Plugins/Experimental/MeshModelingToolset/Source/MeshModelingToolsEditorOnly/Public/MergeMeshesTool.h b/Engine/Plugins/Experimental/MeshModelingToolset/Source/MeshModelingToolsEditorOnly/Public/MergeMeshesTool.h
--- a/Engine/Plugins/Experimental/MeshModelingToolset/Source/MeshModelingToolsEditorOnly/Public/MergeMeshesTool.h
+++ b/Engine/Plugins/Experimental/MeshModelingToolset/Source/MeshModelingToolsEditorOnly/Public/MergeMeshesTool.h
@@ -11,6 +11,8 @@
 #include "PropertySets/OnAcceptProperties.h"
 #include "MergeMeshesTool.generated.h"
 
+class IAssetGenerationAPI;
+
 
 /**
  *
@@ -21,7 +23,7 @@ class MESHMODELINGTOOLSEDITORONLY_API UMergeMeshesToolBuilder : public UInteract
 	GENERATED_BODY()
 
 public:
-	IToolsContextAssetAPI* AssetAPI;
+	IAssetGenerationAPI* AssetAPI;
 
 	UMergeMeshesToolBuilder()
 	{
@@ -72,7 +74,7 @@ public:
 	UMergeMeshesTool();
 
 	virtual void SetWorld(UWorld* World);
-	virtual void SetAssetAPI(IToolsContextAssetAPI* AssetAPI);
+	virtual void SetAssetAPI(IAssetGenerationAPI* AssetAPI);
 
 	virtual void Setup() override;
 	virtual void Shutdown(EToolShutdownType ShutdownType) override;
@@ -104,7 +106,7 @@ protected:
 
 protected:
 	UWorld* TargetWorld;
-	IToolsContextAssetAPI* AssetAPI;
+	IAssetGenerationAPI* AssetAPI;
 
 	TSharedPtr<TArray<IVoxelBasedCSG::FPlacedMesh>> InputMeshes;
 	/** stash copies of the transforms and pointers to the meshes for consumption by merge Op*/
diff --git a/Engine/Plugins/Experimental/MeshModelingToolset/Source/MeshModelingToolsEditorOnly/Public/MeshToVolumeTool.h b/Engine/Plugins/Experimental/MeshModelingToolset/Source/MeshModelingToolsEditorOnly/Public/MeshToVolumeTool.h
--- a/Engine/Plugins/Experimental/MeshModelingToolset/Source/MeshModelingToolsEditorOnly/Public/MeshToVolumeTool.h
+++ b/Engine/Plugins/Experimental/MeshModelingToolset/Source/MeshModelingToolsEditorOnly/Public/MeshToVolumeTool.h
@@ -13,6 +13,8 @@
 #include "PropertySets/OnAcceptProperties.h"
 #include "MeshToVolumeTool.generated.h"
 
+class IAssetGenerationAPI;
+
 /**
  *
  */
@@ -21,7 +23,7 @@ class MESHMODELINGTOOLSEDITORONLY_API UMeshToVolumeToolBuilder : public UInterac
 {
 	GENERATED_BODY()
 public:
-	IToolsContextAssetAPI* AssetAPI = nullptr;
+	IAssetGenerationAPI* AssetAPI = nullptr;
 
 	virtual bool CanBuildTool(const FToolBuilderState & SceneState) const override;
 	virtual UInteractiveTool* BuildTool(const FToolBuilderState & SceneState) const override;
@@ -79,7 +81,7 @@ class MESHMODELINGTOOLSEDITORONLY_API UMeshToVolumeTool : public USingleSelectio
 public:
 	UMeshToVolumeTool();
 
-	virtual void SetAssetAPI(IToolsContextAssetAPI* InAssetAPI) { this->AssetAPI = InAssetAPI; }
+	virtual void SetAssetAPI(IAssetGenerationAPI* InAssetAPI) { this->AssetAPI = InAssetAPI; }
 
 	virtual void Setup() override;
 	virtual void Shutdown(EToolShutdownType ShutdownType) override;
@@ -105,7 +107,7 @@ protected:
 	ULineSetComponent* VolumeEdgesSet;
 
 protected:
-	IToolsContextAssetAPI* AssetAPI = nullptr;
+	IAssetGenerationAPI* AssetAPI = nullptr;
 
 	FDynamicMesh3 InputMesh;
 
diff --git a/Engine/Plugins/Experimental/MeshModelingToolset/Source/MeshModelingToolsEditorOnly/Public/ParameterizeMeshTool.h b/Engine/Plugins/Experimental/MeshModelingToolset/Source/MeshModelingToolsEditorOnly/Public/ParameterizeMeshTool.h
--- a/Engine/Plugins/Experimental/MeshModelingToolset/Source/MeshModelingToolsEditorOnly/Public/ParameterizeMeshTool.h
+++ b/Engine/Plugins/Experimental/MeshModelingToolset/Source/MeshModelingToolsEditorOnly/Public/ParameterizeMeshTool.h
@@ -18,6 +18,7 @@ struct FMeshDescription;
 class USimpleDynamicMeshComponent;
 class UMaterialInterface;
 class UMaterialInstanceDynamic;
+class IAssetGenerationAPI;
 
 
 
@@ -45,7 +46,7 @@ class MESHMODELINGTOOLSEDITORONLY_API UParameterizeMeshToolBuilder : public UInt
 
 public:
 
-	IToolsContextAssetAPI* AssetAPI;
+	IAssetGenerationAPI* AssetAPI;
 	bool bDoAutomaticGlobalUnwrap = false;
 
 	UParameterizeMeshToolBuilder()
@@ -143,7 +144,7 @@ public:
 	UParameterizeMeshTool();
 
 	virtual void SetWorld(UWorld* World);
-	virtual void SetAssetAPI(IToolsContextAssetAPI* AssetAPI);
+	virtual void SetAssetAPI(IAssetGenerationAPI* AssetAPI);
 	virtual void SetUseAutoGlobalParameterizationMode(bool bEnable);
 
 	virtual void Setup() override;
@@ -186,7 +187,7 @@ protected:
 
 protected:
 	UWorld* TargetWorld;
-	IToolsContextAssetAPI* AssetAPI;
+	IAssetGenerationAPI* AssetAPI;
 
 	TSharedPtr<FMeshDescription> InputMesh;
 
diff --git a/Engine/Plugins/Experimental/MeshModelingToolset/Source/MeshModelingToolsEditorOnly/Public/ShapeSprayTool.h b/Engine/Plugins/Experimental/MeshModelingToolset/Source/MeshModelingToolsEditorOnly/Public/ShapeSprayTool.h
--- a/Engine/Plugins/Experimental/MeshModelingToolset/Source/MeshModelingToolsEditorOnly/Public/ShapeSprayTool.h
+++ b/Engine/Plugins/Experimental/MeshModelingToolset/Source/MeshModelingToolsEditorOnly/Public/ShapeSprayTool.h
@@ -8,6 +8,7 @@
 #include "DynamicMeshBrushTool.h"
 #include "ShapeSprayTool.generated.h"
 
+class IAssetGenerationAPI;
 
 /**
  * UMeshSurfacePointToolBuilder override for UShapeSprayTool
@@ -18,7 +19,7 @@ class MESHMODELINGTOOLSEDITORONLY_API UShapeSprayToolBuilder : public UMeshSurfa
 	GENERATED_BODY()
 
 public:
-	IToolsContextAssetAPI* AssetAPI = nullptr;
+	IAssetGenerationAPI* AssetAPI = nullptr;
 
 	virtual UMeshSurfacePointTool* CreateNewTool(const FToolBuilderState& SceneState) const override;
 };
@@ -70,7 +71,7 @@ public:
 	UShapeSprayTool();
 
 	virtual void SetWorld(UWorld* World);
-	virtual void SetAssetAPI(IToolsContextAssetAPI* AssetAPI);
+	virtual void SetAssetAPI(IAssetGenerationAPI* AssetAPI);
 
 	// UInteractiveTool API
 
@@ -99,7 +100,7 @@ protected:
 
 protected:
 	UWorld* TargetWorld;
-	IToolsContextAssetAPI* AssetAPI;
+	IAssetGenerationAPI* AssetAPI;
 
 	FDynamicMesh3 ShapeMesh;
 	void UpdateShapeMesh();
diff --git a/Engine/Plugins/Experimental/MeshModelingToolset/Source/MeshModelingToolsEditorOnly/Public/UVLayoutTool.h b/Engine/Plugins/Experimental/MeshModelingToolset/Source/MeshModelingToolsEditorOnly/Public/UVLayoutTool.h
--- a/Engine/Plugins/Experimental/MeshModelingToolset/Source/MeshModelingToolsEditorOnly/Public/UVLayoutTool.h
+++ b/Engine/Plugins/Experimental/MeshModelingToolset/Source/MeshModelingToolsEditorOnly/Public/UVLayoutTool.h
@@ -19,6 +19,7 @@
 // predeclarations
 struct FMeshDescription;
 class USimpleDynamicMeshComponent;
+class IAssetGenerationAPI;
 
 
 /**
@@ -30,7 +31,7 @@ class MESHMODELINGTOOLSEDITORONLY_API UUVLayoutToolBuilder : public UInteractive
 	GENERATED_BODY()
 
 public:
-	IToolsContextAssetAPI* AssetAPI = nullptr;
+	IAssetGenerationAPI* AssetAPI = nullptr;
 
 	virtual bool CanBuildTool(const FToolBuilderState& SceneState) const override;
 	virtual UInteractiveTool* BuildTool(const FToolBuilderState& SceneState) const override;
@@ -124,7 +125,7 @@ public:
 	virtual void Shutdown(EToolShutdownType ShutdownType) override;
 
 	virtual void SetWorld(UWorld* World, UInteractiveGizmoManager* GizmoManagerIn);
-	virtual void SetAssetAPI(IToolsContextAssetAPI* AssetAPI);
+	virtual void SetAssetAPI(IAssetGenerationAPI* AssetAPI);
 
 	virtual void OnTick(float DeltaTime) override;
 	virtual void Render(IToolsContextRenderAPI* RenderAPI) override;
@@ -150,7 +151,7 @@ protected:
 	TArray<TSharedPtr<FDynamicMesh3>> OriginalDynamicMeshes;
 
 	UWorld* TargetWorld;
-	IToolsContextAssetAPI* AssetAPI;
+	IAssetGenerationAPI* AssetAPI;
 
 	FViewCameraState CameraState;
 
diff --git a/Engine/Plugins/Experimental/MeshModelingToolset/Source/MeshModelingToolsEditorOnly/Public/VolumeToMeshTool.h b/Engine/Plugins/Experimental/MeshModelingToolset/Source/MeshModelingToolsEditorOnly/Public/VolumeToMeshTool.h
--- a/Engine/Plugins/Experimental/MeshModelingToolset/Source/MeshModelingToolsEditorOnly/Public/VolumeToMeshTool.h
+++ b/Engine/Plugins/Experimental/MeshModelingToolset/Source/MeshModelingToolsEditorOnly/Public/VolumeToMeshTool.h
@@ -11,6 +11,8 @@
 #include "Drawing/LineSetComponent.h"
 #include "VolumeToMeshTool.generated.h"
 
+class IAssetGenerationAPI;
+
 /**
  *
  */
@@ -19,7 +21,7 @@ class MESHMODELINGTOOLSEDITORONLY_API UVolumeToMeshToolBuilder : public UInterac
 {
 	GENERATED_BODY()
 public:
-	IToolsContextAssetAPI* AssetAPI = nullptr;
+	IAssetGenerationAPI* AssetAPI = nullptr;
 
 	virtual bool CanBuildTool(const FToolBuilderState& SceneState) const override;
 	virtual UInteractiveTool* BuildTool(const FToolBuilderState& SceneState) const override;
@@ -72,7 +74,7 @@ public:
 	UVolumeToMeshTool();
 
 	virtual void SetWorld(UWorld* World) { this->TargetWorld = World; }
-	virtual void SetAssetAPI(IToolsContextAssetAPI* InAssetAPI) { this->AssetAPI = InAssetAPI; }
+	virtual void SetAssetAPI(IAssetGenerationAPI* InAssetAPI) { this->AssetAPI = InAssetAPI; }
 	virtual void SetSelection(AVolume* Volume);
 
 	virtual void Setup() override;
@@ -100,7 +102,7 @@ protected:
 
 protected:
 	UWorld* TargetWorld = nullptr;
-	IToolsContextAssetAPI* AssetAPI = nullptr;
+	IAssetGenerationAPI* AssetAPI = nullptr;
 
 	FDynamicMesh3 CurrentMesh;
 
diff --git a/Engine/Plugins/Experimental/MeshModelingToolset/Source/MeshModelingToolsEditorOnly/Public/VoxelCSGMeshesTool.h b/Engine/Plugins/Experimental/MeshModelingToolset/Source/MeshModelingToolsEditorOnly/Public/VoxelCSGMeshesTool.h
--- a/Engine/Plugins/Experimental/MeshModelingToolset/Source/MeshModelingToolsEditorOnly/Public/VoxelCSGMeshesTool.h
+++ b/Engine/Plugins/Experimental/MeshModelingToolset/Source/MeshModelingToolsEditorOnly/Public/VoxelCSGMeshesTool.h
@@ -13,6 +13,8 @@
 
 #include "VoxelCSGMeshesTool.generated.h"
 
+class IAssetGenerationAPI;
+
 
 /**
  *
@@ -23,7 +25,7 @@ class MESHMODELINGTOOLSEDITORONLY_API UVoxelCSGMeshesToolBuilder : public UInter
 	GENERATED_BODY()
 
 public:
-	IToolsContextAssetAPI* AssetAPI;
+	IAssetGenerationAPI* AssetAPI;
 
 	UVoxelCSGMeshesToolBuilder()
 	{
@@ -98,7 +100,7 @@ public:
 	UVoxelCSGMeshesTool();
 
 	virtual void SetWorld(UWorld* World);
-	virtual void SetAssetAPI(IToolsContextAssetAPI* AssetAPI);
+	virtual void SetAssetAPI(IAssetGenerationAPI* AssetAPI);
 
 	virtual void Setup() override;
 	virtual void Shutdown(EToolShutdownType ShutdownType) override;
@@ -131,7 +133,7 @@ protected:
 
 protected:
 	UWorld* TargetWorld;
-	IToolsContextAssetAPI* AssetAPI;
+	IAssetGenerationAPI* AssetAPI;
 
 	TSharedPtr<TArray<IVoxelBasedCSG::FPlacedMesh>> InputMeshes;
 	/** stash copies of the transforms and pointers to the meshes for consumption by the CSG Op*/
diff --git a/Engine/Plugins/Experimental/MeshModelingToolset/Source/ModelingComponents/Private/AssetGenerationUtil.cpp b/Engine/Plugins/Experimental/MeshModelingToolset/Source/ModelingComponents/Private/AssetGenerationUtil.cpp
--- a/Engine/Plugins/Experimental/MeshModelingToolset/Source/ModelingComponents/Private/AssetGenerationUtil.cpp
+++ b/Engine/Plugins/Experimental/MeshModelingToolset/Source/ModelingComponents/Private/AssetGenerationUtil.cpp
@@ -10,7 +10,7 @@
 
 // single material case
 AActor* AssetGenerationUtil::GenerateStaticMeshActor(
-	IToolsContextAssetAPI* AssetAPI,
+	IAssetGenerationAPI* AssetAPI,
 	UWorld* TargetWorld,
 	const FDynamicMesh3* Mesh,
 	const FTransform3d& Transform,
@@ -23,7 +23,7 @@ AActor* AssetGenerationUtil::GenerateStaticMeshActor(
 
 // N-material case
 AActor* AssetGenerationUtil::GenerateStaticMeshActor(
-	IToolsContextAssetAPI* AssetAPI,
+	IAssetGenerationAPI* AssetAPI,
 	UWorld* TargetWorld,
 	const FDynamicMesh3* Mesh,
 	const FTransform3d& Transform,
@@ -50,7 +50,7 @@ AActor* AssetGenerationUtil::GenerateStaticMeshActor(
 
 
 bool AssetGenerationUtil::SaveGeneratedTexture2D(
-	IToolsContextAssetAPI* AssetAPI,
+	IAssetGenerationAPI* AssetAPI,
 	UTexture2D* TransientTexture,
 	FString ObjectName,
 	const UObject* RelativeToAsset)
diff --git a/Engine/Plugins/Experimental/MeshModelingToolset/Source/ModelingComponents/Private/BaseTools/BaseCreateFromSelectedTool.cpp b/Engine/Plugins/Experimental/MeshModelingToolset/Source/ModelingComponents/Private/BaseTools/BaseCreateFromSelectedTool.cpp
--- a/Engine/Plugins/Experimental/MeshModelingToolset/Source/ModelingComponents/Private/BaseTools/BaseCreateFromSelectedTool.cpp
+++ b/Engine/Plugins/Experimental/MeshModelingToolset/Source/ModelingComponents/Private/BaseTools/BaseCreateFromSelectedTool.cpp
@@ -102,7 +102,7 @@ void UBaseCreateFromSelectedTool::SetWorld(UWorld* World)
 }
 
 
-void UBaseCreateFromSelectedTool::SetAssetAPI(IToolsContextAssetAPI* AssetAPIIn)
+void UBaseCreateFromSelectedTool::SetAssetAPI(IAssetGenerationAPI* AssetAPIIn)
 {
 	this->AssetAPI = AssetAPIIn;
 }
diff --git a/Engine/Plugins/Experimental/MeshModelingToolset/Source/ModelingComponents/Public/AssetGenerationUtil.h b/Engine/Plugins/Experimental/MeshModelingToolset/Source/ModelingComponents/Public/AssetGenerationUtil.h
--- a/Engine/Plugins/Experimental/MeshModelingToolset/Source/ModelingComponents/Public/AssetGenerationUtil.h
+++ b/Engine/Plugins/Experimental/MeshModelingToolset/Source/ModelingComponents/Public/AssetGenerationUtil.h
@@ -4,15 +4,50 @@
 
 #include "CoreMinimal.h"
 #include "TransformTypes.h"
+#include "ToolContextInterfaces.h"
 
 class UMaterialInterface;
 class UInteractiveToolManager;
-class IToolsContextAssetAPI;
 struct FMeshDescription;
 class FDynamicMesh3;
 class AActor;
 class UTexture2D;
 
+MODELINGCOMPONENTS_API class IAssetGenerationAPI
+{
+public:
+	/**
+	 * Create a new UStaticMeshAsset for the provided mesh, then a new UStaticMeshComponent and AStaticMeshActor in the TargetWorld
+	 * @param TargetWorld world to create new Actor in
+	 * @param Transform transform to set on Actor
+	 * @param ObjectBaseName a base name for the Asset and Actor. Unique name will be generated by appending a number, if a name collision occurs.
+	 * @param AssetConfig configuration information for the new Asset
+	 * @return new AStaticMeshActor with new Asset assigned to it's UStaticMeshComponent
+	 * @warning this function may return null. In this case the asset creation failed and/or the user cancelled during the process, if it was interactive.
+	 */
+	virtual AActor* GenerateStaticMeshActor(
+		UWorld* TargetWorld,
+		FTransform Transform,
+		FString ObjectBaseName,
+		FGeneratedStaticMeshAssetConfig&& AssetConfig) = 0;
+
+
+	/**
+	 * Save a generated UTexture2D as an Asset.
+	 * Assumption is that the texture was generated in code and is in the Transient package.
+	 * @param GeneratedTexture Texture2D to save
+	 * @param ObjectBaseName a base name for the Asset. Unique name will be generated by appending a number, if a name collision occurs.
+	 * @param RelativeToAsset New texture will be saved at the same path as the UPackage for this UObject, which we assume to be an Asset (eg like a UStaticMesh)
+	 */
+	virtual bool SaveGeneratedTexture2D(
+		UTexture2D* GeneratedTexture,
+		FString ObjectBaseName,
+		const UObject* RelativeToAsset)
+	{
+		check(false); return false;
+	}
+};
+
 /**
  * Utility functions for generating Assets in InteractiveTool implementations
  */
@@ -30,7 +65,7 @@ namespace AssetGenerationUtil
 	 * @return new mesh actor
 	 */
 	MODELINGCOMPONENTS_API AActor* GenerateStaticMeshActor(
-		IToolsContextAssetAPI* AssetAPI,
+		IAssetGenerationAPI* AssetAPI,
 		UWorld* TargetWorld,
 		const FDynamicMesh3* Mesh,
 		const FTransform3d& Transform,
@@ -50,7 +85,7 @@ namespace AssetGenerationUtil
 	* @return new mesh actor
 	*/
 	MODELINGCOMPONENTS_API AActor* GenerateStaticMeshActor(
-		IToolsContextAssetAPI* AssetAPI,
+		IAssetGenerationAPI* AssetAPI,
 		UWorld* TargetWorld,
 		const FDynamicMesh3* Mesh,
 		const FTransform3d& Transform,
@@ -67,7 +102,7 @@ namespace AssetGenerationUtil
 	 * @param RelativeToAsset New texture will be saved at the same path as the UPackage for this UObject, which we assume to be an Asset (eg like a UStaticMesh)
 	 */
 	MODELINGCOMPONENTS_API bool SaveGeneratedTexture2D(
-		IToolsContextAssetAPI* AssetAPI,
+		IAssetGenerationAPI* AssetAPI,
 		UTexture2D* TransientTexture,
 		FString ObjectName,
 		const UObject* RelativeToAsset);
diff --git a/Engine/Plugins/Experimental/MeshModelingToolset/Source/ModelingComponents/Public/BaseTools/BaseCreateFromSelectedTool.h b/Engine/Plugins/Experimental/MeshModelingToolset/Source/ModelingComponents/Public/BaseTools/BaseCreateFromSelectedTool.h
--- a/Engine/Plugins/Experimental/MeshModelingToolset/Source/ModelingComponents/Public/BaseTools/BaseCreateFromSelectedTool.h
+++ b/Engine/Plugins/Experimental/MeshModelingToolset/Source/ModelingComponents/Public/BaseTools/BaseCreateFromSelectedTool.h
@@ -10,6 +10,7 @@
 #include "PropertySets/OnAcceptProperties.h"
 #include "BaseCreateFromSelectedTool.generated.h"
 
+class IAssetGenerationAPI;
 
 /**
  * ToolBuilder for UBaseCreateFromSelectedTool
@@ -24,7 +25,7 @@ public:
 	virtual UInteractiveTool* BuildTool(const FToolBuilderState& SceneState) const override;
 
 public:
-	IToolsContextAssetAPI* AssetAPI = nullptr;
+	IAssetGenerationAPI* AssetAPI = nullptr;
 	
 	virtual TOptional<int32> MaxComponentsSupported() const { return TOptional<int32>(); }
 	virtual int32 MinComponentsSupported() const { return 1; }
@@ -67,7 +68,7 @@ public:
 	UBaseCreateFromSelectedTool() = default;
 
 	virtual void SetWorld(UWorld* World);
-	virtual void SetAssetAPI(IToolsContextAssetAPI* AssetAPI);
+	virtual void SetAssetAPI(IAssetGenerationAPI* AssetAPI);
 
 	//
 	// InteractiveTool API - generally does not need to be modified by subclasses
@@ -162,7 +163,7 @@ protected:
 	TArray<UTransformGizmo*> TransformGizmos;
 
 	UWorld* TargetWorld = nullptr;
-	IToolsContextAssetAPI* AssetAPI = nullptr;
+	IAssetGenerationAPI* AssetAPI = nullptr;
 
 };
 
diff --git a/Engine/Plugins/Experimental/ModelingToolsEditorMode/Source/ModelingToolsEditorMode/Private/ModelingModeAssetAPI.cpp b/Engine/Plugins/Experimental/ModelingToolsEditorMode/Source/ModelingToolsEditorMode/Private/ModelingModeAssetAPI.cpp
--- a/Engine/Plugins/Experimental/ModelingToolsEditorMode/Source/ModelingToolsEditorMode/Private/ModelingModeAssetAPI.cpp
+++ b/Engine/Plugins/Experimental/ModelingToolsEditorMode/Source/ModelingToolsEditorMode/Private/ModelingModeAssetAPI.cpp
@@ -5,6 +5,7 @@
 #include "MeshDescription.h"
 #include "Engine/Classes/Engine/Texture2D.h"
 #include "Misc/Guid.h"
+#include "ToolContextInterfaces.h"
 
 #include "StaticMeshComponentBuilder.h"
 
@@ -15,7 +16,11 @@
 
 
 
-
+FModelingModeAssetAPI::FModelingModeAssetAPI(IToolsContextAssetAPI* InAssetAPI)
+	: AssetAPI(InAssetAPI)
+{
+	check(AssetAPI);
+}
 
 bool FModelingModeAssetAPI::GetNewActorPackagePath(
 	UWorld* TargetWorld,
@@ -32,7 +37,7 @@ bool FModelingModeAssetAPI::GetNewActorPackagePath(
 	FString PackageFolderPath = TEXT("/Game/");
 	if (Settings->AssetGenerationLocation == EModelingModeAssetGenerationLocation::AutoGeneratedWorldRelativeAssetPath)
 	{
-		PackageFolderPath = GetWorldRelativeAssetRootPath(TargetWorld);
+		PackageFolderPath = AssetAPI->GetWorldRelativeAssetRootPath(TargetWorld);
 	}
 
 	// combine with fixed AutoGen path name if it is not empty
@@ -54,7 +59,7 @@ bool FModelingModeAssetAPI::GetNewActorPackagePath(
 	// if we want to use the currently-visible asset browser path, try to find one (this can fail if no asset browser is visible/etc)
 	if (Settings->AssetGenerationLocation == EModelingModeAssetGenerationLocation::CurrentAssetBrowserPathIfAvailable)
 	{
-		FString CurrentAssetPath = GetActiveAssetFolderPath();
+		FString CurrentAssetPath = AssetAPI->GetActiveAssetFolderPath();
 		if (CurrentAssetPath.IsEmpty() == false)
 		{
 			PackageFolderPath = CurrentAssetPath;
@@ -65,7 +70,7 @@ bool FModelingModeAssetAPI::GetNewActorPackagePath(
 	// If the user cancels, we are going to discard the asset
 	if (AutoGenMode == EModelingModeAssetGenerationBehavior::InteractivePromptToSave)
 	{
-		FString SelectedPath = InteractiveSelectAssetPath(ObjectBaseName, LOCTEXT("GenerateStaticMeshActorPathDialogWarning", "Choose Folder Path and Name for New Asset. Cancel to Discard New Asset."));
+		FString SelectedPath = AssetAPI->InteractiveSelectAssetPath(ObjectBaseName, LOCTEXT("GenerateStaticMeshActorPathDialogWarning", "Choose Folder Path and Name for New Asset. Cancel to Discard New Asset."));
 		if (SelectedPath.IsEmpty() == false)
 		{
 			PackageFolderPath = FPaths::GetPath(SelectedPath);
@@ -114,7 +119,7 @@ AActor* FModelingModeAssetAPI::GenerateStaticMeshActor(
 
 	// create new package with unique local name
 	FString UniqueAssetName;
-	UPackage* AssetPackage = MakeNewAssetPackage(UsePackageFolderPath, UseBaseName, UniqueAssetName);
+	UPackage* AssetPackage = AssetAPI->MakeNewAssetPackage(UsePackageFolderPath, UseBaseName, UniqueAssetName);
 
 	// create new actor
 	FRotator Rotation(0.0f, 0.0f, 0.0f);
@@ -149,17 +154,17 @@ AActor* FModelingModeAssetAPI::GenerateStaticMeshActor(
 	// save the new asset (or don't, if that's what the user wants)
 	if (AutoGenMode == EModelingModeAssetGenerationBehavior::AutoGenerateAndAutosave )
 	{ 
-		AutoSaveGeneratedAsset(Builder.NewStaticMesh, AssetPackage);
+		AssetAPI->AutoSaveGeneratedAsset(Builder.NewStaticMesh, AssetPackage);
 	}
 	else if (AutoGenMode == EModelingModeAssetGenerationBehavior::InteractivePromptToSave)
 	{
-		AutoSaveGeneratedAsset(Builder.NewStaticMesh, AssetPackage);
+		AssetAPI->AutoSaveGeneratedAsset(Builder.NewStaticMesh, AssetPackage);
 		// this spawns a dialog that just allows save or not-save, but not renaming/etc, seems kind of useless...
 		//InteractiveSaveGeneratedAsset(Builder.NewStaticMesh, AssetPackage);
 	}
 	else if (AutoGenMode == EModelingModeAssetGenerationBehavior::AutoGenerateButDoNotAutosave)
 	{
-		NotifyGeneratedAssetModified(Builder.NewStaticMesh, AssetPackage);
+		AssetAPI->NotifyGeneratedAssetModified(Builder.NewStaticMesh, AssetPackage);
 	}
 
 	return NewActor;
@@ -197,7 +202,7 @@ bool FModelingModeAssetAPI::SaveGeneratedTexture2D(
 	// If the user cancels, we are going to discard the asset
 	if (AutoGenMode == EModelingModeAssetGenerationBehavior::InteractivePromptToSave)
 	{
-		FString SelectedPath = InteractiveSelectAssetPath(ObjectBaseName, LOCTEXT("GenerateTexture2DAssetPathDialogWarning", "Choose Folder Path and Name for New Asset. Cancel to Discard New Asset."));
+		FString SelectedPath = AssetAPI->InteractiveSelectAssetPath(ObjectBaseName, LOCTEXT("GenerateTexture2DAssetPathDialogWarning", "Choose Folder Path and Name for New Asset. Cancel to Discard New Asset."));
 		if (SelectedPath.IsEmpty() == false)
 		{
 			PackageFolderPath = FPaths::GetPath(SelectedPath);
@@ -220,7 +225,7 @@ bool FModelingModeAssetAPI::SaveGeneratedTexture2D(
 
 	// create new package
 	FString UniqueAssetName;
-	UPackage* AssetPackage = MakeNewAssetPackage(PackageFolderPath, UseBaseName, UniqueAssetName);
+	UPackage* AssetPackage = AssetAPI->MakeNewAssetPackage(PackageFolderPath, UseBaseName, UniqueAssetName);
 
 	// move texture from Transient package to real package
 	GeneratedTexture->Rename( *UniqueAssetName, AssetPackage, REN_None);
@@ -237,17 +242,17 @@ bool FModelingModeAssetAPI::SaveGeneratedTexture2D(
 	// save the new asset (or don't, if that's what the user wants)
 	if (AutoGenMode == EModelingModeAssetGenerationBehavior::AutoGenerateAndAutosave)
 	{
-		AutoSaveGeneratedAsset(GeneratedTexture, AssetPackage);
+		AssetAPI->AutoSaveGeneratedAsset(GeneratedTexture, AssetPackage);
 	}
 	else if (AutoGenMode == EModelingModeAssetGenerationBehavior::InteractivePromptToSave)
 	{
-		AutoSaveGeneratedAsset(GeneratedTexture, AssetPackage);
+		AssetAPI->AutoSaveGeneratedAsset(GeneratedTexture, AssetPackage);
 		// this spawns a dialog that just allows save or not-save, but not renaming/etc, seems kind of useless...
 		//InteractiveSaveGeneratedAsset(Builder.NewStaticMesh, AssetPackage);
 	}
 	else if (AutoGenMode == EModelingModeAssetGenerationBehavior::AutoGenerateButDoNotAutosave)
 	{
-		NotifyGeneratedAssetModified(GeneratedTexture, AssetPackage);
+		AssetAPI->NotifyGeneratedAssetModified(GeneratedTexture, AssetPackage);
 	}
 
 	return true;
diff --git a/Engine/Plugins/Experimental/ModelingToolsEditorMode/Source/ModelingToolsEditorMode/Private/ModelingToolsEditorMode.cpp b/Engine/Plugins/Experimental/ModelingToolsEditorMode/Source/ModelingToolsEditorMode/Private/ModelingToolsEditorMode.cpp
--- a/Engine/Plugins/Experimental/ModelingToolsEditorMode/Source/ModelingToolsEditorMode/Private/ModelingToolsEditorMode.cpp
+++ b/Engine/Plugins/Experimental/ModelingToolsEditorMode/Source/ModelingToolsEditorMode/Private/ModelingToolsEditorMode.cpp
@@ -100,11 +100,7 @@ FModelingToolsEditorMode::FModelingToolsEditorMode()
 
 FModelingToolsEditorMode::~FModelingToolsEditorMode()
 {
-	if (ToolsContext != nullptr)
-	{
-		ToolsContext->ShutdownContext();
-		ToolsContext = nullptr;
-	}
+	ToolsContext = nullptr;
 }
 
 
@@ -417,19 +413,8 @@ void FModelingToolsEditorMode::Enter()
 	FEdMode::Enter();
 
 	// initialize FEdMode ToolsContext adapter
-
-	//ToolsContext = NewObject<UEdModeInteractiveToolsContext>(GetTransientPackage(), TEXT("ToolsContext"), RF_Transient);
-	ToolsContext = NewObject<UEdModeInteractiveToolsContext>();
-	ToolsContext->InitializeContextFromEdMode(this, new FModelingModeAssetAPI);
-
-	ToolsContext->OnToolNotificationMessage.AddLambda([this](const FText& Message)
-	{
-		this->OnToolNotificationMessage.Broadcast(Message);
-	});
-	ToolsContext->OnToolWarningMessage.AddLambda([this](const FText& Message)
-	{
-		this->OnToolWarningMessage.Broadcast(Message);
-	});
+	ToolsContext = Owner->GetInteractiveToolsContext();
+	ModelingModeAssetGenerationAPI = MakeShareable(new FModelingModeAssetAPI(ToolsContext->GetAssetAPI()));
 
 	// register stylus event handler
 	StylusStateTracker = MakeUnique<FStylusStateTracker>();
@@ -491,6 +476,8 @@ void FModelingToolsEditorMode::Enter()
 		CommandList->MapAction( UICommand,
 			FExecuteAction::CreateLambda([this, ToolIdentifier]() { ToolsContext->StartTool(ToolIdentifier); }),
 			FCanExecuteAction::CreateLambda([this, ToolIdentifier]() { return ToolsContext->CanStartTool(ToolIdentifier); }));
+
+		RegisteredTools.Emplace(UICommand, ToolIdentifier);
 	};
 
 
@@ -505,7 +492,7 @@ void FModelingToolsEditorMode::Enter()
 								  UAddPrimitiveToolBuilder::EMakeMeshShapeType ShapeTypeIn)
 	{
 		auto AddPrimitiveToolBuilder = NewObject<UAddPrimitiveToolBuilder>();
-		AddPrimitiveToolBuilder->AssetAPI = ToolsContext->GetAssetAPI();
+		AddPrimitiveToolBuilder->AssetAPI = ModelingModeAssetGenerationAPI.Get();
 		AddPrimitiveToolBuilder->ShapeType = ShapeTypeIn;
 		RegisterToolFunc(UICommand, ToolIdentifier, AddPrimitiveToolBuilder);
 	};
@@ -547,27 +534,27 @@ void FModelingToolsEditorMode::Enter()
 	// make shape tools
 	//
 	auto AddPatchToolBuilder = NewObject<UAddPatchToolBuilder>();
-	AddPatchToolBuilder->AssetAPI = ToolsContext->GetAssetAPI();
+	AddPatchToolBuilder->AssetAPI = ModelingModeAssetGenerationAPI.Get();
 	RegisterToolFunc(ToolManagerCommands.BeginAddPatchTool, TEXT("AddPatchTool"), AddPatchToolBuilder);
 
 	auto RevolveBoundaryToolBuilder = NewObject<URevolveBoundaryToolBuilder>();
-	RevolveBoundaryToolBuilder->AssetAPI = ToolsContext->GetAssetAPI();
+	RevolveBoundaryToolBuilder->AssetAPI = ModelingModeAssetGenerationAPI.Get();
 	RegisterToolFunc(ToolManagerCommands.BeginRevolveBoundaryTool, TEXT("RevolveBoundaryTool"), RevolveBoundaryToolBuilder);
 
 	auto DrawPolygonToolBuilder = NewObject<UDrawPolygonToolBuilder>();
-	DrawPolygonToolBuilder->AssetAPI = ToolsContext->GetAssetAPI();
+	DrawPolygonToolBuilder->AssetAPI = ModelingModeAssetGenerationAPI.Get();
 	RegisterToolFunc(ToolManagerCommands.BeginDrawPolygonTool, TEXT("DrawPolygonTool"), DrawPolygonToolBuilder);
 
 	auto DrawPolyPathToolBuilder = NewObject<UDrawPolyPathToolBuilder>();
-	DrawPolyPathToolBuilder->AssetAPI = ToolsContext->GetAssetAPI();
+	DrawPolyPathToolBuilder->AssetAPI = ModelingModeAssetGenerationAPI.Get();
 	RegisterToolFunc(ToolManagerCommands.BeginDrawPolyPathTool, TEXT("DrawPolyPath"), DrawPolyPathToolBuilder);
 
 	auto DrawAndRevolveToolBuilder = NewObject<UDrawAndRevolveToolBuilder>();
-	DrawAndRevolveToolBuilder->AssetAPI = ToolsContext->GetAssetAPI();
+	DrawAndRevolveToolBuilder->AssetAPI = ModelingModeAssetGenerationAPI.Get();
 	RegisterToolFunc(ToolManagerCommands.BeginDrawAndRevolveTool, TEXT("RevolveTool"), DrawAndRevolveToolBuilder);
 
 	auto ShapeSprayToolBuilder = NewObject<UShapeSprayToolBuilder>();
-	ShapeSprayToolBuilder->AssetAPI = ToolsContext->GetAssetAPI();
+	ShapeSprayToolBuilder->AssetAPI = ModelingModeAssetGenerationAPI.Get();
 	RegisterToolFunc(ToolManagerCommands.BeginShapeSprayTool, TEXT("ShapeSprayTool"), ShapeSprayToolBuilder);
 
 
@@ -595,11 +582,11 @@ void FModelingToolsEditorMode::Enter()
 	RegisterToolFunc(ToolManagerCommands.BeginTransformUVIslandsTool, TEXT("EditUVIslands"), NewObject<UEditUVIslandsToolBuilder>());
 
 	UCombineMeshesToolBuilder* CombineMeshesToolBuilder = NewObject<UCombineMeshesToolBuilder>();
-	CombineMeshesToolBuilder->AssetAPI = ToolsContext->GetAssetAPI();
+	CombineMeshesToolBuilder->AssetAPI = ModelingModeAssetGenerationAPI.Get();
 	RegisterToolFunc(ToolManagerCommands.BeginCombineMeshesTool, TEXT("CombineMeshesTool"), CombineMeshesToolBuilder);
 
 	UCombineMeshesToolBuilder* DuplicateMeshesToolBuilder = NewObject<UCombineMeshesToolBuilder>();
-	DuplicateMeshesToolBuilder->AssetAPI = ToolsContext->GetAssetAPI();
+	DuplicateMeshesToolBuilder->AssetAPI = ModelingModeAssetGenerationAPI.Get();
 	DuplicateMeshesToolBuilder->bIsDuplicateTool = true;
 	RegisterToolFunc(ToolManagerCommands.BeginDuplicateMeshesTool, TEXT("DuplicateMeshesTool"), DuplicateMeshesToolBuilder);
 
@@ -621,100 +608,100 @@ void FModelingToolsEditorMode::Enter()
 	RegisterToolFunc(ToolManagerCommands.BeginEdgeLoopInsertionTool, TEXT("EdgeLoopInsertionTool"), EdgeLoopInsertionToolBuilder);
 
 	auto EditNormalsToolBuilder = NewObject<UEditNormalsToolBuilder>();
-	EditNormalsToolBuilder->AssetAPI = ToolsContext->GetAssetAPI();
+	EditNormalsToolBuilder->AssetAPI = ModelingModeAssetGenerationAPI.Get();
 	RegisterToolFunc(ToolManagerCommands.BeginEditNormalsTool, TEXT("EditNormalsTool"), EditNormalsToolBuilder);
 
 	auto TangentsToolBuilder = NewObject<UMeshTangentsToolBuilder>();
 	RegisterToolFunc(ToolManagerCommands.BeginEditTangentsTool, TEXT("MeshTangentsTool"), TangentsToolBuilder);
 
 	auto RemoveOccludedTrianglesToolBuilder = NewObject<URemoveOccludedTrianglesToolBuilder>();
-	RemoveOccludedTrianglesToolBuilder->AssetAPI = ToolsContext->GetAssetAPI();
+	RemoveOccludedTrianglesToolBuilder->AssetAPI = ModelingModeAssetGenerationAPI.Get();
 	RegisterToolFunc(ToolManagerCommands.BeginRemoveOccludedTrianglesTool, TEXT("RemoveOccludedTrianglesTool"), RemoveOccludedTrianglesToolBuilder);
 
 	auto HoleFillToolBuilder = NewObject<UHoleFillToolBuilder>();
 	RegisterToolFunc(ToolManagerCommands.BeginHoleFillTool, TEXT("HoleFillTool"), HoleFillToolBuilder);
 
 	auto UVProjectionToolBuilder = NewObject<UUVProjectionToolBuilder>();
-	UVProjectionToolBuilder->AssetAPI = ToolsContext->GetAssetAPI();
+	UVProjectionToolBuilder->AssetAPI = ModelingModeAssetGenerationAPI.Get();
 	RegisterToolFunc(ToolManagerCommands.BeginUVProjectionTool, TEXT("UVProjectionTool"), UVProjectionToolBuilder);
 
 	auto UVLayoutToolBuilder = NewObject<UUVLayoutToolBuilder>();
-	UVLayoutToolBuilder->AssetAPI = ToolsContext->GetAssetAPI();
+	UVLayoutToolBuilder->AssetAPI = ModelingModeAssetGenerationAPI.Get();
 	RegisterToolFunc(ToolManagerCommands.BeginUVLayoutTool, TEXT("UVLayoutTool"), UVLayoutToolBuilder);
 
 	auto MergeMeshesToolBuilder = NewObject<UMergeMeshesToolBuilder>();
-	MergeMeshesToolBuilder->AssetAPI = ToolsContext->GetAssetAPI();
+	MergeMeshesToolBuilder->AssetAPI = ModelingModeAssetGenerationAPI.Get();
 	RegisterToolFunc(ToolManagerCommands.BeginVoxelMergeTool, TEXT("MergeMeshesTool"), MergeMeshesToolBuilder);
 
 	auto VoxelCSGMeshesToolBuilder = NewObject<UVoxelCSGMeshesToolBuilder>();
-	VoxelCSGMeshesToolBuilder->AssetAPI = ToolsContext->GetAssetAPI();
+	VoxelCSGMeshesToolBuilder->AssetAPI = ModelingModeAssetGenerationAPI.Get();
 	RegisterToolFunc(ToolManagerCommands.BeginVoxelBooleanTool, TEXT("VoxelCSGMeshesTool"), VoxelCSGMeshesToolBuilder);
 
 	auto VoxelSolidifyMeshesToolBuilder = NewObject<UVoxelSolidifyMeshesToolBuilder>();
-	VoxelSolidifyMeshesToolBuilder->AssetAPI = ToolsContext->GetAssetAPI();
+	VoxelSolidifyMeshesToolBuilder->AssetAPI = ModelingModeAssetGenerationAPI.Get();
 	RegisterToolFunc(ToolManagerCommands.BeginVoxelSolidifyTool, TEXT("VoxelSolidifyMeshesTool"), VoxelSolidifyMeshesToolBuilder);
 
 	auto VoxelBlendMeshesToolBuilder = NewObject<UVoxelBlendMeshesToolBuilder>();
-	VoxelBlendMeshesToolBuilder->AssetAPI = ToolsContext->GetAssetAPI();
+	VoxelBlendMeshesToolBuilder->AssetAPI = ModelingModeAssetGenerationAPI.Get();
 	RegisterToolFunc(ToolManagerCommands.BeginVoxelBlendTool, TEXT("VoxelBlendMeshesTool"), VoxelBlendMeshesToolBuilder);
 
 	auto VoxelMorphologyMeshesToolBuilder = NewObject<UVoxelMorphologyMeshesToolBuilder>();
-	VoxelMorphologyMeshesToolBuilder->AssetAPI = ToolsContext->GetAssetAPI();
+	VoxelMorphologyMeshesToolBuilder->AssetAPI = ModelingModeAssetGenerationAPI.Get();
 	RegisterToolFunc(ToolManagerCommands.BeginVoxelMorphologyTool, TEXT("VoxelMorphologyMeshesTool"), VoxelMorphologyMeshesToolBuilder);
 
 	auto SelfUnionMeshesToolBuilder = NewObject<USelfUnionMeshesToolBuilder>();
-	SelfUnionMeshesToolBuilder->AssetAPI = ToolsContext->GetAssetAPI();
+	SelfUnionMeshesToolBuilder->AssetAPI = ModelingModeAssetGenerationAPI.Get();
 	RegisterToolFunc(ToolManagerCommands.BeginSelfUnionTool, TEXT("SelfUnionMeshesTool"), SelfUnionMeshesToolBuilder);
 
 	auto CSGMeshesToolBuilder = NewObject<UCSGMeshesToolBuilder>();
-	CSGMeshesToolBuilder->AssetAPI = ToolsContext->GetAssetAPI();
+	CSGMeshesToolBuilder->AssetAPI = ModelingModeAssetGenerationAPI.Get();
 	RegisterToolFunc(ToolManagerCommands.BeginMeshBooleanTool, TEXT("CSGMeshesTool"), CSGMeshesToolBuilder);
 
 	auto BspConversionToolBuilder = NewObject<UBspConversionToolBuilder>();
-	BspConversionToolBuilder->AssetAPI = ToolsContext->GetAssetAPI();
+	BspConversionToolBuilder->AssetAPI = ModelingModeAssetGenerationAPI.Get();
 	RegisterToolFunc(ToolManagerCommands.BeginBspConversionTool, TEXT("BspConversionTool"), BspConversionToolBuilder);
 
 	auto MeshToVolumeToolBuilder = NewObject<UMeshToVolumeToolBuilder>();
-	MeshToVolumeToolBuilder->AssetAPI = ToolsContext->GetAssetAPI();
+	MeshToVolumeToolBuilder->AssetAPI = ModelingModeAssetGenerationAPI.Get();
 	RegisterToolFunc(ToolManagerCommands.BeginMeshToVolumeTool, TEXT("MeshToVolumeTool"), MeshToVolumeToolBuilder);
 
 	auto VolumeToMeshToolBuilder = NewObject<UVolumeToMeshToolBuilder>();
-	VolumeToMeshToolBuilder->AssetAPI = ToolsContext->GetAssetAPI();
+	VolumeToMeshToolBuilder->AssetAPI = ModelingModeAssetGenerationAPI.Get();
 	RegisterToolFunc(ToolManagerCommands.BeginVolumeToMeshTool, TEXT("VolumeToMeshTool"), VolumeToMeshToolBuilder);
 
 	auto PlaneCutToolBuilder = NewObject<UPlaneCutToolBuilder>();
-	PlaneCutToolBuilder->AssetAPI = ToolsContext->GetAssetAPI();
+	PlaneCutToolBuilder->AssetAPI = ModelingModeAssetGenerationAPI.Get();
 	RegisterToolFunc(ToolManagerCommands.BeginPlaneCutTool, TEXT("PlaneCutTool"), PlaneCutToolBuilder);
 
 	auto MirrorToolBuilder = NewObject<UMirrorToolBuilder>();
-	MirrorToolBuilder->AssetAPI = ToolsContext->GetAssetAPI();
+	MirrorToolBuilder->AssetAPI = ModelingModeAssetGenerationAPI.Get();
 	RegisterToolFunc(ToolManagerCommands.BeginMirrorTool, TEXT("MirrorTool"), MirrorToolBuilder);
 
 	auto PolygonCutToolBuilder = NewObject<UPolygonOnMeshToolBuilder>();
 	RegisterToolFunc(ToolManagerCommands.BeginPolygonCutTool, TEXT("PolyCutTool"), PolygonCutToolBuilder);
 
 	auto GlobalUVGenerateToolBuilder = NewObject<UParameterizeMeshToolBuilder>();
-	GlobalUVGenerateToolBuilder->AssetAPI = ToolsContext->GetAssetAPI();
+	GlobalUVGenerateToolBuilder->AssetAPI = ModelingModeAssetGenerationAPI.Get();
 	GlobalUVGenerateToolBuilder->bDoAutomaticGlobalUnwrap = true;
 	RegisterToolFunc(ToolManagerCommands.BeginGlobalUVGenerateTool, TEXT("GlobalParameterizeMeshTool"), GlobalUVGenerateToolBuilder);
 
 	auto GroupUVGenerateToolBuilder = NewObject<UParameterizeMeshToolBuilder>();
-	GroupUVGenerateToolBuilder->AssetAPI = ToolsContext->GetAssetAPI();
+	GroupUVGenerateToolBuilder->AssetAPI = ModelingModeAssetGenerationAPI.Get();
 	GroupUVGenerateToolBuilder->bDoAutomaticGlobalUnwrap = false;
 	RegisterToolFunc(ToolManagerCommands.BeginGroupUVGenerateTool, TEXT("GroupParameterizeMeshTool"), GroupUVGenerateToolBuilder);
 
 	auto MeshSelectionToolBuilder = NewObject<UMeshSelectionToolBuilder>();
-	MeshSelectionToolBuilder->AssetAPI = ToolsContext->GetAssetAPI();
+	MeshSelectionToolBuilder->AssetAPI = ModelingModeAssetGenerationAPI.Get();
 	RegisterToolFunc(ToolManagerCommands.BeginMeshSelectionTool, TEXT("MeshSelectionTool"), MeshSelectionToolBuilder);
 
 	auto EditMeshMaterialsToolBuilder = NewObject<UEditMeshMaterialsToolBuilder>();
-	EditMeshMaterialsToolBuilder->AssetAPI = ToolsContext->GetAssetAPI();
+	EditMeshMaterialsToolBuilder->AssetAPI = ModelingModeAssetGenerationAPI.Get();
 	RegisterToolFunc(ToolManagerCommands.BeginEditMeshMaterialsTool, TEXT("EditMaterialsTool"), EditMeshMaterialsToolBuilder);
 	
 	RegisterToolFunc(ToolManagerCommands.BeginMeshAttributePaintTool, TEXT("MeshAttributePaintTool"), NewObject<UMeshAttributePaintToolBuilder>());
 
 	auto BakeMeshAttributeMapsToolBuilder = NewObject<UBakeMeshAttributeMapsToolBuilder>();
-	BakeMeshAttributeMapsToolBuilder->AssetAPI = ToolsContext->GetAssetAPI();
+	BakeMeshAttributeMapsToolBuilder->AssetAPI = ModelingModeAssetGenerationAPI.Get();
 	RegisterToolFunc(ToolManagerCommands.BeginBakeMeshAttributeMapsTool, TEXT("BakeMeshMapsTool"), BakeMeshAttributeMapsToolBuilder);
 
 	// analysis tools
@@ -732,7 +719,7 @@ void FModelingToolsEditorMode::Enter()
 	//RegisterToolFunc(ToolManagerCommands.BeginEditCollisionGeometryTool, TEXT("EditCollisionGeoTool"), NewObject<UEditCollisionGeometryToolBuilder>());
 
 	auto ExtractCollisionGeoToolBuilder = NewObject<UExtractCollisionGeometryToolBuilder>();
-	ExtractCollisionGeoToolBuilder->AssetAPI = ToolsContext->GetAssetAPI();
+	ExtractCollisionGeoToolBuilder->AssetAPI = ModelingModeAssetGenerationAPI.Get();
 	RegisterToolFunc(ToolManagerCommands.BeginExtractCollisionGeometryTool, TEXT("ExtractCollisionGeoTool"), ExtractCollisionGeoToolBuilder);
 
 
@@ -762,14 +749,15 @@ void FModelingToolsEditorMode::Enter()
 
 void FModelingToolsEditorMode::Exit()
 {
-	OnToolNotificationMessage.Clear();
-	OnToolWarningMessage.Clear();
-
-
 	StylusStateTracker = nullptr;
 
-	ToolsContext->ShutdownContext();
-	ToolsContext = nullptr;
+	const TSharedRef<FUICommandList>& CommandList = Toolkit->GetToolkitCommands();
+	for (auto& RegisteredTool : RegisteredTools)
+	{
+		CommandList->UnmapAction(RegisteredTool.Key);
+		ToolsContext->ToolManager->UnregisterToolType(RegisteredTool.Value);
+	}
+	RegisteredTools.SetNum(0);
 
 	if (Toolkit.IsValid())
 	{
@@ -784,6 +772,9 @@ void FModelingToolsEditorMode::Exit()
 		Toolkit.Reset();
 	}
 
+	ToolsContext = nullptr;
+	ModelingModeAssetGenerationAPI.Reset();
+
 	FModelingModeActionCommands::UnRegisterCommandBindings(UICommandList);
 
 	// clear realtime viewport override
diff --git a/Engine/Plugins/Experimental/ModelingToolsEditorMode/Source/ModelingToolsEditorMode/Private/ModelingToolsEditorModeToolkit.cpp b/Engine/Plugins/Experimental/ModelingToolsEditorMode/Source/ModelingToolsEditorMode/Private/ModelingToolsEditorModeToolkit.cpp
--- a/Engine/Plugins/Experimental/ModelingToolsEditorMode/Source/ModelingToolsEditorMode/Private/ModelingToolsEditorModeToolkit.cpp
+++ b/Engine/Plugins/Experimental/ModelingToolsEditorMode/Source/ModelingToolsEditorMode/Private/ModelingToolsEditorModeToolkit.cpp
@@ -39,6 +39,8 @@ FModelingToolsEditorModeToolkit::~FModelingToolsEditorModeToolkit()
 {
 	UModelingToolsEditorModeSettings* Settings = GetMutableDefault<UModelingToolsEditorModeSettings>();
 	Settings->OnModified.Remove(AssetSettingsModifiedHandle);
+	GetToolsEditorMode()->GetToolsContext()->OnToolNotificationMessage.RemoveAll(this);
+	GetToolsEditorMode()->GetToolsContext()->OnToolWarningMessage.RemoveAll(this);
 }
 
 void FModelingToolsEditorModeToolkit::Init(const TSharedPtr<IToolkitHost>& InitToolkitHost)
@@ -121,35 +123,8 @@ void FModelingToolsEditorModeToolkit::Init(const TSharedPtr<IToolkitHost>& InitT
 	ActiveToolName = FText::GetEmpty();
 	ActiveToolMessage = FText::GetEmpty();
 
-	GetToolsEditorMode()->GetToolManager()->OnToolStarted.AddLambda([this](UInteractiveToolManager* Manager, UInteractiveTool* Tool)
-	{
-		UpdateActiveToolProperties();
-
-		UInteractiveTool* CurTool = GetToolsEditorMode()->GetToolManager()->GetActiveTool(EToolSide::Left);
-		CurTool->OnPropertySetsModified.AddLambda([this]() { UpdateActiveToolProperties(); });
-
-		ModeHeaderArea->SetVisibility(EVisibility::Collapsed);
-		ActiveToolName = CurTool->GetToolInfo().ToolDisplayName;
-	});
-	GetToolsEditorMode()->GetToolManager()->OnToolEnded.AddLambda([this](UInteractiveToolManager* Manager, UInteractiveTool* Tool)
-	{
-		DetailsView->SetObject(nullptr);
-		ActiveToolName = FText::GetEmpty();
-		ModeHeaderArea->SetVisibility(EVisibility::Visible);
-		ModeHeaderArea->SetText(LOCTEXT("SelectToolLabel", "Select a Tool from the Toolbar"));
-		ClearNotification();
-		ClearWarning();
-	});
-
-
-	GetToolsEditorMode()->OnToolNotificationMessage.AddLambda([this](const FText& Message)
-	{
-		PostNotification(Message);
-	});
-	GetToolsEditorMode()->OnToolWarningMessage.AddLambda([this](const FText& Message)
-	{
-		PostWarning(Message);
-	});
+	GetToolsEditorMode()->GetToolsContext()->OnToolNotificationMessage.AddSP(this, &FModelingToolsEditorModeToolkit::PostNotification);
+	GetToolsEditorMode()->GetToolsContext()->OnToolWarningMessage.AddSP(this, &FModelingToolsEditorModeToolkit::PostWarning);
 }
 
 
@@ -667,6 +642,29 @@ void FModelingToolsEditorModeToolkit::EnableShowRealtimeWarning(bool bEnable)
 	}
 }
 
+void FModelingToolsEditorModeToolkit::OnToolStarted(UInteractiveToolManager* Manager, UInteractiveTool* Tool)
+{
+	UpdateActiveToolProperties();
+
+	UInteractiveTool* CurTool = GetToolsEditorMode()->GetToolManager()->GetActiveTool(EToolSide::Left);
+	CurTool->OnPropertySetsModified.AddSP(this, &FModelingToolsEditorModeToolkit::UpdateActiveToolProperties);
+
+	ModeHeaderArea->SetVisibility(EVisibility::Collapsed);
+	ActiveToolName = CurTool->GetToolInfo().ToolDisplayName;
+}
+
+void FModelingToolsEditorModeToolkit::OnToolEnded(UInteractiveToolManager* Manager, UInteractiveTool* Tool)
+{
+	DetailsView->SetObject(nullptr);
+	UInteractiveTool* CurTool = GetToolsEditorMode()->GetToolManager()->GetActiveTool(EToolSide::Left);
+	ActiveToolName = FText::GetEmpty();
+	ModeHeaderArea->SetVisibility(EVisibility::Visible);
+	ModeHeaderArea->SetText(LOCTEXT("SelectToolLabel", "Select a Tool from the Toolbar"));
+	ClearNotification();
+	ClearWarning();
+	CurTool->OnPropertySetsModified.RemoveAll(this);
+}
+
 void FModelingToolsEditorModeToolkit::UpdateShowWarnings()
 {
 	if (bShowRealtimeWarning )
diff --git a/Engine/Plugins/Experimental/ModelingToolsEditorMode/Source/ModelingToolsEditorMode/Public/ModelingModeAssetAPI.h b/Engine/Plugins/Experimental/ModelingToolsEditorMode/Source/ModelingToolsEditorMode/Public/ModelingModeAssetAPI.h
--- a/Engine/Plugins/Experimental/ModelingToolsEditorMode/Source/ModelingToolsEditorMode/Public/ModelingModeAssetAPI.h
+++ b/Engine/Plugins/Experimental/ModelingToolsEditorMode/Source/ModelingToolsEditorMode/Public/ModelingModeAssetAPI.h
@@ -3,15 +3,20 @@
 #pragma once
 
 #include "CoreMinimal.h"
-#include "Tools/EditorToolAssetAPI.h"
+#include "AssetGenerationUtil.h"
+
+class IToolsContextAssetAPI;
+struct FGeneratedStaticMeshAssetConfig;
 
 /**
  * Implementation of ToolsContext Asset management API that is suitable for use
  * inside Modeling Tools Mode. 
  */
-class MODELINGTOOLSEDITORMODE_API FModelingModeAssetAPI : public FEditorToolAssetAPI
+class MODELINGTOOLSEDITORMODE_API FModelingModeAssetAPI : public IAssetGenerationAPI
 {
 public:
+	FModelingModeAssetAPI(IToolsContextAssetAPI* InAssetAPI);
+
 	/**
 	 * Generate static mesh actor
 	 */
@@ -39,4 +44,7 @@ public:
 		FString& PackageFolderPathOut,
 		FString& ObjectBaseNameOut
 		);
+
+protected:
+	IToolsContextAssetAPI* AssetAPI;
 };
\ No newline at end of file
diff --git a/Engine/Plugins/Experimental/ModelingToolsEditorMode/Source/ModelingToolsEditorMode/Public/ModelingToolsEditorMode.h b/Engine/Plugins/Experimental/ModelingToolsEditorMode/Source/ModelingToolsEditorMode/Public/ModelingToolsEditorMode.h
--- a/Engine/Plugins/Experimental/ModelingToolsEditorMode/Source/ModelingToolsEditorMode/Public/ModelingToolsEditorMode.h
+++ b/Engine/Plugins/Experimental/ModelingToolsEditorMode/Source/ModelingToolsEditorMode/Public/ModelingToolsEditorMode.h
@@ -15,6 +15,7 @@ class FEditorComponentSourceFactory;
 class FEditorToolAssetAPI;
 class FUICommandList;
 class FStylusStateTracker;		// for stylus events
+class FModelingModeAssetAPI;
 
 class FModelingToolsEditorMode : public FEdMode
 {
@@ -134,14 +135,6 @@ public:
 	// End of FEdMode interface
 	//////////////////
 
-
-
-public:
-	DECLARE_MULTICAST_DELEGATE_OneParam(FOnModelingModeToolNotification, const FText&);
-	FOnModelingModeToolNotification OnToolNotificationMessage;
-	FOnModelingModeToolNotification OnToolWarningMessage;
-
-
 public:
 	virtual UEdModeInteractiveToolsContext* GetToolsContext() const
 	{
@@ -160,8 +153,12 @@ protected:
 	/** Command list lives here so that the key bindings on the commands can be processed in the viewport. */
 	TSharedPtr<FUICommandList> UICommandList;
 
+	TSharedPtr<FModelingModeAssetAPI> ModelingModeAssetGenerationAPI;
+
 	TUniquePtr<FStylusStateTracker> StylusStateTracker;
 
+	TArray<TPair<TSharedPtr<FUICommandInfo>, FString>> RegisteredTools;
+
 	void ModelingModeShortcutRequested(EModelingModeActionCommands Command);
 	void FocusCameraAtCursorHotkey();
 
diff --git a/Engine/Plugins/Experimental/ModelingToolsEditorMode/Source/ModelingToolsEditorMode/Public/ModelingToolsEditorModeToolkit.h b/Engine/Plugins/Experimental/ModelingToolsEditorMode/Source/ModelingToolsEditorMode/Public/ModelingToolsEditorModeToolkit.h
--- a/Engine/Plugins/Experimental/ModelingToolsEditorMode/Source/ModelingToolsEditorMode/Public/ModelingToolsEditorModeToolkit.h
+++ b/Engine/Plugins/Experimental/ModelingToolsEditorMode/Source/ModelingToolsEditorMode/Public/ModelingToolsEditorModeToolkit.h
@@ -57,6 +57,8 @@ public:
 
 	virtual void EnableShowRealtimeWarning(bool bEnable);
 
+	virtual void OnToolStarted(UInteractiveToolManager* Manager, UInteractiveTool* Tool) override;
+	virtual void OnToolEnded(UInteractiveToolManager* Manager, UInteractiveTool* Tool) override;
 
 
 
diff --git a/Engine/Plugins/Experimental/SampleToolsEditorMode/Source/Private/SampleToolsEditorMode.cpp b/Engine/Plugins/Experimental/SampleToolsEditorMode/Source/Private/SampleToolsEditorMode.cpp
--- a/Engine/Plugins/Experimental/SampleToolsEditorMode/Source/Private/SampleToolsEditorMode.cpp
+++ b/Engine/Plugins/Experimental/SampleToolsEditorMode/Source/Private/SampleToolsEditorMode.cpp
@@ -2,235 +2,79 @@
 
 #include "SampleToolsEditorMode.h"
 #include "SampleToolsEditorModeToolkit.h"
-#include "Toolkits/ToolkitManager.h"
-#include "EditorViewportClient.h"
-#include "EditorModeManager.h"
+#include "EdModeInteractiveToolsContext.h"
+#include "InteractiveToolManager.h"
+#include "SampleToolseditorModeCommands.h"
 
 
 //////////////////////////////////////////////////////////////////////////
 ////////////////////////////////////////////////////////////////////////// 
-// AddYourTool Step 1 - include the header file for your Tool here
+// AddYourTool Step 1 - include the header file for your Tools here
 //////////////////////////////////////////////////////////////////////////
 ////////////////////////////////////////////////////////////////////////// 
-#include "BaseTools/SingleClickTool.h"
-#include "BaseTools/MeshSurfacePointTool.h"
 #include "SampleTools/CreateActorSampleTool.h"
 #include "SampleTools/DrawCurveOnMeshSampleTool.h"
 #include "SampleTools/MeasureDistanceSampleTool.h"
 
 // step 2: register a ToolBuilder in FSampleToolsEditorMode::Enter()
-// step 3: add a button in FSampleToolsEditorModeToolkit::Init()
 
 
 #define LOCTEXT_NAMESPACE "FSampleToolsEditorMode"
 
-const FEditorModeID FSampleToolsEditorMode::EM_SampleToolsEditorModeId = TEXT("EM_SampleToolsEditorMode");
+const FEditorModeID USampleToolsEditorMode::EM_SampleToolsEditorModeId = TEXT("EM_SampleToolsEditorMode");
 
 
-FSampleToolsEditorMode::FSampleToolsEditorMode()
+USampleToolsEditorMode::USampleToolsEditorMode()
 {
-	ToolsContext = nullptr;
+	Info = FEditorModeInfo(USampleToolsEditorMode::EM_SampleToolsEditorModeId,
+		LOCTEXT("SampleToolsEditorModeName", "SampleToolsEditorMode"),
+		FSlateIcon(),
+		true);
 }
 
 
-FSampleToolsEditorMode::~FSampleToolsEditorMode()
+USampleToolsEditorMode::~USampleToolsEditorMode()
 {
-	// this should have happend already in ::Exit()
-	if (ToolsContext != nullptr)
-	{
-		ToolsContext->ShutdownContext();
-		ToolsContext = nullptr;
-	}
 }
 
 
-void FSampleToolsEditorMode::ActorSelectionChangeNotify()
+void USampleToolsEditorMode::ActorSelectionChangeNotify()
 {
 	// @todo support selection change
 }
 
-
-
-void FSampleToolsEditorMode::Tick(FEditorViewportClient* ViewportClient, float DeltaTime)
-{
-	FEdMode::Tick(ViewportClient, DeltaTime);
-
-	// give ToolsContext a chance to tick
-	if (ToolsContext != nullptr)
-	{
-		ToolsContext->Tick(ViewportClient, DeltaTime);
-	}
-}
-
-
-
-
-void FSampleToolsEditorMode::Render(const FSceneView* View, FViewport* Viewport, FPrimitiveDrawInterface* PDI)
-{
-	FEdMode::Render(View, Viewport, PDI);
-
-	// give ToolsContext a chance to render
-	if (ToolsContext != nullptr)
-	{
-		ToolsContext->Render(View, Viewport, PDI);
-	}
-}
-
-
-
-
-
-//
-// Input device event tracking. We forward input events to the ToolsContext adapter for handling.
-// 
-
-bool FSampleToolsEditorMode::InputKey(FEditorViewportClient* ViewportClient, FViewport* Viewport, FKey Key, EInputEvent Event)
-{
-	bool bHandled = FEdMode::InputKey(ViewportClient, Viewport, Key, Event);
-	bHandled |= ToolsContext->InputKey(ViewportClient, Viewport, Key, Event);
-	return bHandled;
-}
-
-bool FSampleToolsEditorMode::InputAxis(FEditorViewportClient* InViewportClient, FViewport* Viewport, int32 ControllerId, FKey Key, float Delta, float DeltaTime)
-{
-	// mouse axes: EKeys::MouseX, EKeys::MouseY, EKeys::MouseWheelAxis
-	return FEdMode::InputAxis(InViewportClient, Viewport, ControllerId, Key, Delta, DeltaTime);
-}
-
-
-bool FSampleToolsEditorMode::StartTracking(FEditorViewportClient* InViewportClient, FViewport* InViewport) 
-{ 
-	bool bHandled = FEdMode::StartTracking(InViewportClient, InViewport);
-	bHandled |= ToolsContext->StartTracking(InViewportClient, InViewport);
-	return bHandled;
-}
-
-bool FSampleToolsEditorMode::CapturedMouseMove(FEditorViewportClient* InViewportClient, FViewport* InViewport, int32 InMouseX, int32 InMouseY)
-{
-	bool bHandled = ToolsContext->CapturedMouseMove(InViewportClient, InViewport, InMouseX, InMouseY);
-	return bHandled;
-}
-
-bool FSampleToolsEditorMode::EndTracking(FEditorViewportClient* InViewportClient, FViewport* InViewport)
+void USampleToolsEditorMode::Enter()
 {
-	bool bHandled = ToolsContext->EndTracking(InViewportClient, InViewport);
-	return bHandled;
-}
-
-
-
-bool FSampleToolsEditorMode::MouseEnter(FEditorViewportClient* ViewportClient, FViewport* Viewport, int32 x, int32 y)
-{
-	bool bHandled = ToolsContext->MouseEnter(ViewportClient, Viewport, x, y);
-	return bHandled;
-}
-
-bool FSampleToolsEditorMode::MouseMove(FEditorViewportClient* ViewportClient, FViewport* Viewport, int32 x, int32 y)
-{
-	bool bHandled = ToolsContext->MouseMove(ViewportClient, Viewport, x, y);
-	return bHandled;
-}
-
-bool FSampleToolsEditorMode::MouseLeave(FEditorViewportClient* ViewportClient, FViewport* Viewport)
-{
-	bool bHandled = ToolsContext->MouseLeave(ViewportClient, Viewport);
-	return bHandled;
-}
-
-
-
-bool FSampleToolsEditorMode::ReceivedFocus(FEditorViewportClient* ViewportClient, FViewport* Viewport)
-{
-	return FEdMode::ReceivedFocus(ViewportClient, Viewport);
-}
-
-
-bool FSampleToolsEditorMode::LostFocus(FEditorViewportClient* ViewportClient, FViewport* Viewport)
-{
-	return FEdMode::LostFocus(ViewportClient, Viewport);
-}
-
-
-
-
-
-
-void FSampleToolsEditorMode::Enter()
-{
-	FEdMode::Enter();
-
-	if (!Toolkit.IsValid() && UsesToolkits())
-	{
-		Toolkit = MakeShareable(new FSampleToolsEditorModeToolkit);
-		Toolkit->Init(Owner->GetToolkitHost());
-	}
-
-	// initialize the adapter that attaches the ToolsContext to this FEdMode
-	ToolsContext = NewObject<UEdModeInteractiveToolsContext>(GetTransientPackage(), TEXT("ToolsContext"), RF_Transient);
-	ToolsContext->InitializeContextFromEdMode(this);
-
+	UEdMode::Enter();
 
 	//////////////////////////////////////////////////////////////////////////
 	//////////////////////////////////////////////////////////////////////////
-	// AddYourTool Step 2 - register a ToolBuilder for your Tool here.
+	// AddYourTool Step 2 - register the ToolBuilders for your Tools here.
 	// The string name you pass to the ToolManager is used to select/activate your ToolBuilder later.
 	//////////////////////////////////////////////////////////////////////////
 	////////////////////////////////////////////////////////////////////////// 
+	FSampleToolsEditorModeCommands SampleToolCommands = FSampleToolsEditorModeCommands::Get();
 
-
-	auto CreateActorSampleToolBuilder = NewObject<UCreateActorSampleToolBuilder>();
+	auto CreateActorSampleToolBuilder = NewObject< UCreateActorSampleToolBuilder>(this);
 	CreateActorSampleToolBuilder->AssetAPI = ToolsContext->GetAssetAPI();
-	ToolsContext->ToolManager->RegisterToolType(TEXT("CreateActorSampleTool"), CreateActorSampleToolBuilder);
-
-	auto DrawCurveOnMeshSampleToolBuilder = NewObject<UDrawCurveOnMeshSampleToolBuilder>();
-	ToolsContext->ToolManager->RegisterToolType(TEXT("DrawCurveOnMeshSampleTool"), DrawCurveOnMeshSampleToolBuilder);
+	RegisterTool(SampleToolCommands.CreateActorTool, TEXT("CreateActorSampleTool"), CreateActorSampleToolBuilder);
 
-	auto MeasureDistanceSampleToolBuilder = NewObject<UMeasureDistanceSampleToolBuilder>();
-	ToolsContext->ToolManager->RegisterToolType(TEXT("MeasureDistanceSampleTool"), MeasureDistanceSampleToolBuilder);
-
-	auto SurfacePointToolBuilder = NewObject<UMeshSurfacePointToolBuilder>();
-	ToolsContext->ToolManager->RegisterToolType(TEXT("SurfacePointTool"), SurfacePointToolBuilder);
+	RegisterTool(SampleToolCommands.DrawCurveOnMeshTool, TEXT("DrawCurveOnMeshSampleTool"), NewObject<UDrawCurveOnMeshSampleToolBuilder>(this));
+	RegisterTool(SampleToolCommands.MeasureDistanceTool, TEXT("MeasureDistanceSampleTool"), NewObject<UMeasureDistanceSampleToolBuilder>(this));
+	RegisterTool(SampleToolCommands.SurfacePointTool, TEXT("SurfacePointTool"), NewObject<UMeshSurfacePointToolBuilder>(this));
 
 	// active tool type is not relevant here, we just set to default
-	ToolsContext->ToolManager->SelectActiveToolType(EToolSide::Left, TEXT("SurfacePointTool"));
+	GetToolManager()->SelectActiveToolType(EToolSide::Left, TEXT("SurfacePointTool"));
 }
 
-
-
-
-void FSampleToolsEditorMode::Exit()
+void USampleToolsEditorMode::CreateToolkit()
 {
-	// shutdown and clean up the ToolsContext
-	ToolsContext->ShutdownContext();
-	ToolsContext = nullptr;
-
-	if (Toolkit.IsValid())
-	{
-		FToolkitManager::Get().CloseToolkit(Toolkit.ToSharedRef());
-		Toolkit.Reset();
-	}
-
-	// Call base Exit method to ensure proper cleanup
-	FEdMode::Exit();
+	Toolkit = MakeShareable(new FSampleToolsEditorModeToolkit);
 }
 
-bool FSampleToolsEditorMode::UsesToolkits() const
+TMap<FName, TArray<TSharedPtr<FUICommandInfo>>> USampleToolsEditorMode::GetModeCommands() const
 {
-	return true;
+	return FSampleToolsEditorModeCommands::Get().GetCommands();
 }
 
-
-void FSampleToolsEditorMode::AddReferencedObjects(FReferenceCollector& Collector)
-{
-	Collector.AddReferencedObject(ToolsContext);
-}
-
-
-
-
-
-
-
-
-
-#undef LOCTEXT_NAMESPACE
\ No newline at end of file
+#undef LOCTEXT_NAMESPACE
diff --git a/Engine/Plugins/Experimental/SampleToolsEditorMode/Source/Private/SampleToolsEditorModeCommands.cpp b/Engine/Plugins/Experimental/SampleToolsEditorMode/Source/Private/SampleToolsEditorModeCommands.cpp
new file mode 100644
--- /dev/null
+++ b/Engine/Plugins/Experimental/SampleToolsEditorMode/Source/Private/SampleToolsEditorModeCommands.cpp
@@ -0,0 +1,34 @@
+// Copyright Epic Games, Inc. All Rights Reserved.
+
+#include "SampleToolsEditorModeCommands.h"
+#include "EditorStyleSet.h"
+
+#define LOCTEXT_NAMESPACE "FSampleToolsEditorMode"
+
+FSampleToolsEditorModeCommands::FSampleToolsEditorModeCommands()
+	: TCommands<FSampleToolsEditorModeCommands>("SampleToolsEditorMode",
+		NSLOCTEXT("SampleToolsEditorMode", "SampleToolsEditorModeCommands", "Sample Tools Editor Mode"),
+		NAME_None,
+		FEditorStyle::GetStyleSetName())
+{
+}
+
+void FSampleToolsEditorModeCommands::RegisterCommands()
+{
+	TArray <TSharedPtr<FUICommandInfo>>& ToolCommands = Commands.FindOrAdd(NAME_Default);
+	UI_COMMAND(CreateActorTool, "Create Actor on Click", "Create Actor on Click", EUserInterfaceActionType::ToggleButton, FInputChord());
+	ToolCommands.Add(CreateActorTool);
+	UI_COMMAND(DrawCurveOnMeshTool, "Draw Curve On Mesh", "Draw Curve On Mesh", EUserInterfaceActionType::ToggleButton, FInputChord());
+	ToolCommands.Add(DrawCurveOnMeshTool);
+	UI_COMMAND(MeasureDistanceTool, "Measure Distance", "Measure Distance", EUserInterfaceActionType::ToggleButton, FInputChord());
+	ToolCommands.Add(MeasureDistanceTool);
+	UI_COMMAND(SurfacePointTool, "Surface Point Tool", "Surface Point Tool", EUserInterfaceActionType::ToggleButton, FInputChord());
+	ToolCommands.Add(SurfacePointTool);
+}
+
+TMap<FName, TArray<TSharedPtr<FUICommandInfo>>> FSampleToolsEditorModeCommands::GetCommands()
+{
+	return FSampleToolsEditorModeCommands::Get().Commands;
+}
+
+#undef LOCTEXT_NAMESPACE
diff --git a/Engine/Plugins/Experimental/SampleToolsEditorMode/Source/Private/SampleToolsEditorModeModule.cpp b/Engine/Plugins/Experimental/SampleToolsEditorMode/Source/Private/SampleToolsEditorModeModule.cpp
--- a/Engine/Plugins/Experimental/SampleToolsEditorMode/Source/Private/SampleToolsEditorModeModule.cpp
+++ b/Engine/Plugins/Experimental/SampleToolsEditorMode/Source/Private/SampleToolsEditorModeModule.cpp
@@ -1,21 +1,23 @@
 // Copyright Epic Games, Inc. All Rights Reserved.
 
 #include "SampleToolsEditorModeModule.h"
-#include "SampleToolsEditorMode.h"
+#include "SampleToolsEditorModeCommands.h"
 
 #define LOCTEXT_NAMESPACE "FSampleToolsEditorModeModule"
 
 void FSampleToolsEditorModeModule::StartupModule()
 {
 	// This code will execute after your module is loaded into memory; the exact timing is specified in the .uplugin file per-module
-	FEditorModeRegistry::Get().RegisterMode<FSampleToolsEditorMode>(FSampleToolsEditorMode::EM_SampleToolsEditorModeId, LOCTEXT("SampleToolsEditorModeName", "SampleToolsEditorMode"), FSlateIcon(), true);
+
+	FSampleToolsEditorModeCommands::Register();
 }
 
 void FSampleToolsEditorModeModule::ShutdownModule()
 {
 	// This function may be called during shutdown to clean up your module.  For modules that support dynamic reloading,
 	// we call this function before unloading the module.
-	FEditorModeRegistry::Get().UnregisterMode(FSampleToolsEditorMode::EM_SampleToolsEditorModeId);
+
+	FSampleToolsEditorModeCommands::Unregister();
 }
 
 #undef LOCTEXT_NAMESPACE
diff --git a/Engine/Plugins/Experimental/SampleToolsEditorMode/Source/Private/SampleToolsEditorModeToolkit.cpp b/Engine/Plugins/Experimental/SampleToolsEditorMode/Source/Private/SampleToolsEditorModeToolkit.cpp
--- a/Engine/Plugins/Experimental/SampleToolsEditorMode/Source/Private/SampleToolsEditorModeToolkit.cpp
+++ b/Engine/Plugins/Experimental/SampleToolsEditorMode/Source/Private/SampleToolsEditorModeToolkit.cpp
@@ -21,181 +21,14 @@ FSampleToolsEditorModeToolkit::FSampleToolsEditorModeToolkit()
 {
 }
 
-void FSampleToolsEditorModeToolkit::Init(const TSharedPtr<IToolkitHost>& InitToolkitHost)
+void FSampleToolsEditorModeToolkit::Init(const TSharedPtr<IToolkitHost>& InitToolkitHost, TWeakObjectPtr<UEdMode> InOwningMode)
 {
-
-	FPropertyEditorModule& PropertyEditorModule = FModuleManager::GetModuleChecked<FPropertyEditorModule>("PropertyEditor");
-
-	FDetailsViewArgs DetailsViewArgs(
-		/*bUpdateFromSelection=*/ false,
-		/*bLockable=*/ false,
-		/*bAllowSearch=*/ false,
-		FDetailsViewArgs::HideNameArea,
-		/*bHideSelectionTip=*/ true,
-		/*InNotifyHook=*/ nullptr,
-		/*InSearchInitialKeyFocus=*/ false,
-		/*InViewIdentifier=*/ NAME_None);
-	DetailsViewArgs.DefaultsOnlyVisibility = EEditDefaultsOnlyNodeVisibility::Automatic;
-	DetailsViewArgs.bShowOptions = false;
-	DetailsViewArgs.bAllowMultipleTopLevelObjects = true;
-
-	DetailsView = PropertyEditorModule.CreateDetailView(DetailsViewArgs);
-
-
-	//////////////////////////////////////////////////////////////////////////
-	////////////////////////////////////////////////////////////////////////// 
-	// AddYourTool Step 3 - add a button to initialize your Tool
-	//////////////////////////////////////////////////////////////////////////
-	////////////////////////////////////////////////////////////////////////// 
-
-
-	SAssignNew(ToolkitWidget, SBorder)
-		.HAlign(HAlign_Center)
-		.Padding(25)
-		[
-			SNew(SVerticalBox)
-
-			+ SVerticalBox::Slot()
-				.AutoHeight()
-				.HAlign(HAlign_Center)
-				.Padding(50)
-				[
-					SNew(STextBlock)
-					.AutoWrapText(true)
-					.Text(LOCTEXT("HeaderLabel", "Sample Tools"))
-				]
-
-			+ SVerticalBox::Slot()
-				.HAlign(HAlign_Center)
-				.AutoHeight()
-				[
-					SNew(SButton).Text(LOCTEXT("CreateActorSampleToolLabel", "Create Actor on Click"))
-					.OnClicked_Lambda([this]() { return this->StartTool(TEXT("CreateActorSampleTool")); })
-					.IsEnabled_Lambda([this]() { return this->CanStartTool(TEXT("CreateActorSampleTool")); })
-				]
-
-			+ SVerticalBox::Slot()
-				.HAlign(HAlign_Center)
-				.AutoHeight()
-				[
-					SNew(SButton).Text(LOCTEXT("MeasureDistanceSampleToolLabel", "Measure Distance"))
-					.OnClicked_Lambda([this]() { return this->StartTool(TEXT("MeasureDistanceSampleTool")); })
-					.IsEnabled_Lambda([this]() { return this->CanStartTool(TEXT("MeasureDistanceSampleTool")); })
-				]
-
-			+ SVerticalBox::Slot()
-				.HAlign(HAlign_Center)
-				.AutoHeight()
-				[
-					SNew(SButton).Text(LOCTEXT("DrawCurveOnMeshSampleToolLabel", "Draw Curve On Mesh"))
-					.OnClicked_Lambda([this]() { return this->StartTool(TEXT("DrawCurveOnMeshSampleTool")); })
-					.IsEnabled_Lambda([this]() { return this->CanStartTool(TEXT("DrawCurveOnMeshSampleTool")); })
-				]
-
-			+ SVerticalBox::Slot()
-				.HAlign(HAlign_Center)
-				.AutoHeight()
-				[
-					SNew(SButton).Text(LOCTEXT("SurfacePointToolLabel", "Surface Point Tool"))
-					.OnClicked_Lambda([this]() { return this->StartTool(TEXT("SurfacePointTool")); })
-					.IsEnabled_Lambda([this]() { return this->CanStartTool(TEXT("SurfacePointTool")); })
-				]
-
-
-			+ SVerticalBox::Slot()
-				.HAlign(HAlign_Center)
-				.AutoHeight()
-				[
-				SNew(SHorizontalBox)
-				+ SHorizontalBox::Slot()
-					.AutoWidth()
-					[
-						SNew(SButton).Text(LOCTEXT("AcceptToolButtonLabel", "Accept"))
-						.OnClicked_Lambda([this]() { return this->EndTool(EToolShutdownType::Accept); })
-						.IsEnabled_Lambda([this]() { return this->CanAcceptActiveTool(); })
-					]
-				+ SHorizontalBox::Slot()
-					.AutoWidth()
-					[
-						SNew(SButton).Text(LOCTEXT("CancelToolButtonLabel", "Cancel"))
-						.OnClicked_Lambda([this]() { return this->EndTool(EToolShutdownType::Cancel); })
-						.IsEnabled_Lambda([this]() { return this->CanCancelActiveTool(); })
-					]
-				+ SHorizontalBox::Slot()
-					.AutoWidth()
-					[
-						SNew(SButton).Text(LOCTEXT("CompletedToolButtonLabel", "Complete"))
-						.OnClicked_Lambda([this]() { return this->EndTool(EToolShutdownType::Completed); })
-						.IsEnabled_Lambda([this]() { return this->CanCompleteActiveTool(); })
-					]
-				]
-
-			+ SVerticalBox::Slot()
-				.HAlign(HAlign_Center)
-				.Padding(2)
-				.AutoHeight()
-				.MaxHeight(500.0f)
-				[
-					DetailsView->AsShared()
-				]
-
-		];
-		
-	FModeToolkit::Init(InitToolkitHost);
-}
-
-
-bool FSampleToolsEditorModeToolkit::CanStartTool(const FString& ToolTypeIdentifier)
-{
-	UInteractiveToolManager* Manager = GetToolsEditorMode()->GetToolManager();
-
-	return (Manager->HasActiveTool(EToolSide::Left) == false) &&
-		(Manager->CanActivateTool(EToolSide::Left, ToolTypeIdentifier) == true);
-}
-
-bool FSampleToolsEditorModeToolkit::CanAcceptActiveTool()
-{
-	return GetToolsEditorMode()->GetToolManager()->CanAcceptActiveTool(EToolSide::Left);
-}
-
-bool FSampleToolsEditorModeToolkit::CanCancelActiveTool()
-{
-	return GetToolsEditorMode()->GetToolManager()->CanCancelActiveTool(EToolSide::Left);
-}
-
-bool FSampleToolsEditorModeToolkit::CanCompleteActiveTool()
-{
-	return GetToolsEditorMode()->GetToolManager()->HasActiveTool(EToolSide::Left) && CanCancelActiveTool() == false;
-}
-
-
-FReply FSampleToolsEditorModeToolkit::StartTool(const FString& ToolTypeIdentifier)
-{
-	if (GetToolsEditorMode()->GetToolManager()->SelectActiveToolType(EToolSide::Left, ToolTypeIdentifier) == false)
-	{
-		UE_LOG(LogTemp, Warning, TEXT("ToolManager: Unknown Tool Type %s"), *ToolTypeIdentifier);
-	}
-	else
-	{
-		UE_LOG(LogTemp, Warning, TEXT("ToolManager: Starting Tool Type %s"), *ToolTypeIdentifier);
-		GetToolsEditorMode()->GetToolManager()->ActivateTool(EToolSide::Left);
-
-		// Update properties panel
-		UInteractiveTool* CurTool = GetToolsEditorMode()->GetToolManager()->GetActiveTool(EToolSide::Left);
-		DetailsView->SetObjects(CurTool->GetToolProperties());
-	}
-	return FReply::Handled();
+	FModeToolkit::Init(InitToolkitHost, InOwningMode);
 }
 
-FReply FSampleToolsEditorModeToolkit::EndTool(EToolShutdownType ShutdownType)
+void FSampleToolsEditorModeToolkit::GetToolPaletteNames(TArray<FName>& PaletteNames) const
 {
-	UE_LOG(LogTemp, Warning, TEXT("ENDING TOOL"));
-
-	GetToolsEditorMode()->GetToolManager()->DeactivateTool(EToolSide::Left, ShutdownType);
-
-	DetailsView->SetObject(nullptr);
-
-	return FReply::Handled();
+	PaletteNames.Add(NAME_Default);
 }
 
 
@@ -209,14 +42,4 @@ FText FSampleToolsEditorModeToolkit::GetBaseToolkitName() const
 	return NSLOCTEXT("SampleToolsEditorModeToolkit", "DisplayName", "SampleToolsEditorMode Tool");
 }
 
-class FEdMode* FSampleToolsEditorModeToolkit::GetEditorMode() const
-{
-	return GLevelEditorModeTools().GetActiveMode(FSampleToolsEditorMode::EM_SampleToolsEditorModeId);
-}
-
-FSampleToolsEditorMode* FSampleToolsEditorModeToolkit::GetToolsEditorMode() const
-{
-	return (FSampleToolsEditorMode*)GetEditorMode();
-}
-
 #undef LOCTEXT_NAMESPACE
diff --git a/Engine/Plugins/Experimental/SampleToolsEditorMode/Source/Public/SampleToolsEditorMode.h b/Engine/Plugins/Experimental/SampleToolsEditorMode/Source/Public/SampleToolsEditorMode.h
--- a/Engine/Plugins/Experimental/SampleToolsEditorMode/Source/Public/SampleToolsEditorMode.h
+++ b/Engine/Plugins/Experimental/SampleToolsEditorMode/Source/Public/SampleToolsEditorMode.h
@@ -3,155 +3,39 @@
 #pragma once
 
 #include "CoreMinimal.h"
-#include "EdMode.h"
+#include "Tools/UEdMode.h"
 
-#include "InputState.h"
-#include "InteractiveToolManager.h"
-#include "EdModeInteractiveToolsContext.h"
-
-
-class FEditorComponentSourceFactory;
-class FEditorToolAssetAPI;
+#include "SampleToolsEditorMode.generated.h"
 
+class IAssetGenerationAPI;
 
 /**
- * This class provides an example of how to connect an FEdMode to the 
- * InteractiveTools framework. The various FEdMode input event handlers 
+ * This class provides an example of how to extend a UEdMode to add some simple tools
+ * using the InteractiveTools framework. The various UEdMode input event handlers (see UEdMode.h)
  * forward events to a UEdModeInteractiveToolsContext instance, which 
  * has all the logic for interacting with the InputRouter, ToolManager, etc.
- * The functions here are basically just forwarders.
+ * The functions provided here are the minimum to get started inserting some custom behavior.
+ * Take a look at the UEdMode markup for more extensibility options.
  */
-class FSampleToolsEditorMode : public FEdMode
+UCLASS()
+class USampleToolsEditorMode : public UEdMode
 {
+	GENERATED_BODY()
+
 public:
 	const static FEditorModeID EM_SampleToolsEditorModeId;
-public:
-	FSampleToolsEditorMode();
-	virtual ~FSampleToolsEditorMode();
-
-	////////////////
-	// FEdMode interface
-	////////////////
-
-	virtual void Tick(FEditorViewportClient* ViewportClient, float DeltaTime) override;
-	virtual void Render(const FSceneView* View, FViewport* Viewport, FPrimitiveDrawInterface* PDI) override;
-	virtual void ActorSelectionChangeNotify() override;
-	virtual bool UsesToolkits() const override;
-
-	virtual void AddReferencedObjects(FReferenceCollector& Collector) override;
 
-	// these disable the standard gizmo, which is probably want we want in
-	// these tools as we can't hit-test the standard gizmo...
-	virtual bool AllowWidgetMove() override { return false; }
-	virtual bool ShouldDrawWidget() const override { return false; }
-	virtual bool UsesTransformWidget() const override { return false; }
+	USampleToolsEditorMode();
+	virtual ~USampleToolsEditorMode();
 
-
-	/*
-	 * focus events
-	 */
-
-	// called when we "start" this editor mode (ie switch to this tab)
+	////////////////
+	// UEdMode interface
 	virtual void Enter() override;
-
-	// called when we "end" this editor mode (ie switch to another tab)
-	virtual void Exit() override;
-
-	// called when viewport window is focused
-	virtual bool ReceivedFocus(FEditorViewportClient* ViewportClient, FViewport* Viewport) override;
-
-	// called when viewport window loses focus (ie when some other window is focused)
-	// *not* called when editor is backgrounded, but is called when editor is minimized
-	virtual bool LostFocus(FEditorViewportClient* ViewportClient, FViewport* Viewport) override;
-
-
-	/*
-	 * Mouse position events - These are called when no mouse button is down
-	 */
-
-	// called when mouse moves over viewport window
-	virtual bool MouseEnter(FEditorViewportClient* ViewportClient, FViewport* Viewport, int32 x, int32 y) override;
-
-	// called when mouse leaves viewport window
-	virtual bool MouseLeave(FEditorViewportClient* ViewportClient, FViewport* Viewport) override;
-
-	// called on any mouse-move event. *not* called during tracking/capturing, eg if any button is down
-	virtual bool MouseMove(FEditorViewportClient* ViewportClient, FViewport* Viewport, int32 x, int32 y) override;
-
-
-	/*
-	 * Input Button/Axis Events & Mouse Capture
-	 *
-	 *   event sequence for left mouse down/capture/up is:
-	 *     - InputKey( EKeys::LeftMouseButton, IE_PRESSED )
-	 *     - StartTracking()
-	 *     - CapturedMouseMove() and InputAxis() (repeated for each mouse-move)
-	 *     - Inputkey( EKeys::LeftMouseButton, IE_RELEASED )
-	 *     - EndTracking()
-	 *
-	 *   for doubleclick, we get one of the above sequence, then a second
-	 *   where instead of the first IE_PRESSED we get a IE_DoubleClick
-	 *
-	 *   for mouse wheel, we get following sequence
-	 *     - InputKey ( EKeys::MouseScrollUp / Down, IE_PRESSED )
-	 *     - InputAxis() for wheel move
-	 *     - InputKey ( EKeys::MouseScrollUp / Down, IE_RELEASED )
-	 *   it appears that there will only ever be *one* InputAxis() between the pressed/released sequenece
-	 *
-	 *   Note that this wheel event sequence can happen *during* a 
-	 *   middle-mouse tracking sequence on a wheel-button (nice!)
-	 */
-
-
-	// This is not just called for keyboard keys, it is also called for mouse down/up events!
-	// Eg for left-press we get Key = EKeys::LeftMouseButton and Event = IE_Pressed
-	// Return value indicates "handled'. If we return true for mouse press events then
-	// the StartTracking/CapturedMouseMove/EndTracking sequence is *not* called (but we
-	// also don't get MouseMove(), so the mouse-movement events appear to be lost?)
-	virtual bool InputKey(FEditorViewportClient* ViewportClient, FViewport* Viewport, FKey Key, EInputEvent Event) override;
-
-
-	// Called for 1D Axis movements - EKeys::MouseX, EKeys::MouseY, and EKeys::MouseWheelAxis
-	// Called even if we return true from InputKey which otherwise blocks tracking. 
-	// Return value indicates "handled" but has no effect on CapturedMouseMove events (ie we still get them)
-	virtual bool InputAxis(FEditorViewportClient* InViewportClient, FViewport* Viewport, int32 ControllerId, FKey Key, float Delta, float DeltaTime);
-
-
-	// called on mouse-down. return value indicates whether this was "handled" but
-	// does not mean we get exclusive capture events!
-	virtual bool StartTracking(FEditorViewportClient* InViewportClient, FViewport* InViewport) override;
-
-	// called during mouse-down mouse-move. Always called, return value is used to indicate whether
-	// we "handled" this mouse move (see FEditorModeTools::CapturedMouseMove) but does not pre-empt
-	// other editor modes from seeing this move event...
-	virtual bool CapturedMouseMove(FEditorViewportClient* InViewportClient, FViewport* InViewport, int32 InMouseX, int32 InMouseY) override;
-
-	// always called on mouse-up
-	virtual bool EndTracking(FEditorViewportClient* InViewportClient, FViewport* InViewport) override;
-
+	virtual void ActorSelectionChangeNotify() override;
+	virtual void CreateToolkit() override;
+	virtual TMap<FName, TArray<TSharedPtr<FUICommandInfo>>> GetModeCommands() const override;
 
 	//////////////////
-	// End of FEdMode interface
+	// End of UEdMode interface
 	//////////////////
-
-
-public:
-
-	/**
-	 * @return active ToolManager
-	 */
-	virtual UInteractiveToolManager* GetToolManager() const
-	{ 
-		return ToolsContext->ToolManager;
-	}
-
-protected:
-
-	/**
-	 * This is the ToolsContext adapter that attaches to the FEdMode. 
-	 * We forward FEdMode event handler functions to this implementation.
-	 * This UObject is explicitly kept alive by adding it to the Collector in ::AddReferencedObjects()
-	 */
-	UEdModeInteractiveToolsContext* ToolsContext;
-
 };
diff --git a/Engine/Plugins/Experimental/SampleToolsEditorMode/Source/Public/SampleToolsEditorModeCommands.h b/Engine/Plugins/Experimental/SampleToolsEditorMode/Source/Public/SampleToolsEditorModeCommands.h
new file mode 100644
--- /dev/null
+++ b/Engine/Plugins/Experimental/SampleToolsEditorMode/Source/Public/SampleToolsEditorModeCommands.h
@@ -0,0 +1,21 @@
+// Copyright Epic Games, Inc. All Rights Reserved.
+
+#include "CoreMinimal.h"
+#include "Framework/Commands/Commands.h"
+
+class FSampleToolsEditorModeCommands : public TCommands<FSampleToolsEditorModeCommands>
+{
+public:
+	FSampleToolsEditorModeCommands();
+
+	virtual void RegisterCommands() override;
+	static TMap<FName, TArray<TSharedPtr<FUICommandInfo>>> GetCommands();
+
+	TSharedPtr<FUICommandInfo> CreateActorTool;
+	TSharedPtr<FUICommandInfo> DrawCurveOnMeshTool;
+	TSharedPtr<FUICommandInfo> MeasureDistanceTool;
+	TSharedPtr<FUICommandInfo> SurfacePointTool;
+
+protected:
+	TMap<FName, TArray<TSharedPtr<FUICommandInfo>>> Commands;
+};
diff --git a/Engine/Plugins/Experimental/SampleToolsEditorMode/Source/Public/SampleToolsEditorModeToolkit.h b/Engine/Plugins/Experimental/SampleToolsEditorMode/Source/Public/SampleToolsEditorModeToolkit.h
--- a/Engine/Plugins/Experimental/SampleToolsEditorMode/Source/Public/SampleToolsEditorModeToolkit.h
+++ b/Engine/Plugins/Experimental/SampleToolsEditorMode/Source/Public/SampleToolsEditorModeToolkit.h
@@ -19,28 +19,10 @@ public:
 	FSampleToolsEditorModeToolkit();
 	
 	// FModeToolkit interface 
-	virtual void Init(const TSharedPtr<IToolkitHost>& InitToolkitHost) override;
+	virtual void Init(const TSharedPtr<IToolkitHost>& InitToolkitHost, TWeakObjectPtr<UEdMode> InOwningMode) override;
+	virtual void GetToolPaletteNames(TArray<FName>& PaletteNames) const override;
 
 	// IToolkit interface 
 	virtual FName GetToolkitFName() const override;
 	virtual FText GetBaseToolkitName() const override;
-	virtual class FEdMode* GetEditorMode() const override;
-	virtual TSharedPtr<class SWidget> GetInlineContent() const override { return ToolkitWidget; }
-
-	virtual FSampleToolsEditorMode* GetToolsEditorMode() const;
-
-private:
-
-	TSharedPtr<SWidget> ToolkitWidget;
-	TSharedPtr<IDetailsView> DetailsView;
-
-	// these functions just forward calls to the ToolsContext / ToolManager
-
-	bool CanStartTool(const FString& ToolTypeIdentifier);
-	bool CanAcceptActiveTool();
-	bool CanCancelActiveTool();
-	bool CanCompleteActiveTool();
-
-	FReply StartTool(const FString& ToolTypeIdentifier);
-	FReply EndTool(EToolShutdownType ShutdownType);
 };
diff --git a/Engine/Plugins/Experimental/SampleToolsEditorMode/Source/SampleToolsEditorMode.Build.cs b/Engine/Plugins/Experimental/SampleToolsEditorMode/Source/SampleToolsEditorMode.Build.cs
--- a/Engine/Plugins/Experimental/SampleToolsEditorMode/Source/SampleToolsEditorMode.Build.cs
+++ b/Engine/Plugins/Experimental/SampleToolsEditorMode/Source/SampleToolsEditorMode.Build.cs
@@ -41,6 +41,7 @@ public class SampleToolsEditorMode : ModuleRules
 				"SlateCore",
 				"InputCore",
 				"EditorFramework",
+				"EditorStyle",
 				"UnrealEd",
                 "ContentBrowser",
                 "LevelEditor",
diff --git a/Engine/Plugins/MeshPainting/Source/MeshPaintEditorMode/Private/BrushSettingsCustomization.cpp b/Engine/Plugins/MeshPainting/Source/MeshPaintEditorMode/Private/BrushSettingsCustomization.cpp
--- a/Engine/Plugins/MeshPainting/Source/MeshPaintEditorMode/Private/BrushSettingsCustomization.cpp
+++ b/Engine/Plugins/MeshPainting/Source/MeshPaintEditorMode/Private/BrushSettingsCustomization.cpp
@@ -435,7 +435,7 @@ void FTexturePaintingSettingsCustomization::CustomizeDetails(IDetailLayoutBuilde
 			.AllowSpin(true)
 			.Value_Lambda([]() -> int32 { return UMeshPaintMode::GetTextureToolProperties() ? UMeshPaintMode::GetTextureToolProperties()->UVChannel : 0; })
 			.MinValue(0)
-			.MaxValue_Lambda([]() -> int32 { return Cast<UMeshToolManager>(UMeshPaintMode::GetMeshPaintMode()->GetToolManager())->GetMaxUVIndexToPaint(); })
+			.MaxValue_Lambda([]() -> int32 { return UMeshPaintMode::GetMeshPaintMode()->GetMeshToolData()->GetMaxUVIndexToPaint(); })
 			.OnValueChanged(SNumericEntryBox<int32>::FOnValueChanged::CreateLambda([=](int32 Value) { 
 				if (UMeshPaintMode::GetTextureToolProperties())
 				{
diff --git a/Engine/Plugins/MeshPainting/Source/MeshPaintEditorMode/Private/MeshPaintMode.cpp b/Engine/Plugins/MeshPainting/Source/MeshPaintEditorMode/Private/MeshPaintMode.cpp
--- a/Engine/Plugins/MeshPainting/Source/MeshPaintEditorMode/Private/MeshPaintMode.cpp
+++ b/Engine/Plugins/MeshPainting/Source/MeshPaintEditorMode/Private/MeshPaintMode.cpp
@@ -118,7 +118,6 @@ UMeshPaintMode::UMeshPaintMode()
 	: Super()
 {
 	SettingsClass = UMeshPaintModeSettings::StaticClass();
-	ToolsContextClass = UMeshToolsContext::StaticClass();
 	// Don't be a visible mode unless legacy mesh paint mode is not on.
 	const bool bVisible = !GetDefault<ULevelEditorMiscSettings>()->bEnableLegacyMeshPaintMode;
 	FModuleManager::Get().LoadModule("EditorStyle");
@@ -139,11 +138,29 @@ void UMeshPaintMode::Enter()
 	ModeSettings = Cast<UMeshPaintModeSettings>(SettingsObject);
 	
 	FMeshPaintEditorModeCommands ToolManagerCommands = FMeshPaintEditorModeCommands::Get();
-	RegisterTool(ToolManagerCommands.VertexSelect, VertexSelectToolName, NewObject<UVertexAdapterClickToolBuilder>());
-	RegisterTool(ToolManagerCommands.TextureSelect, TextureSelectToolName, NewObject<UTextureAdapterClickToolBuilder>());
-	RegisterTool(ToolManagerCommands.ColorPaint, ColorPaintToolName, NewObject<UMeshColorPaintingToolBuilder>());
-	RegisterTool(ToolManagerCommands.WeightPaint, WeightPaintToolName, NewObject<UMeshWeightPaintingToolBuilder>());
-	RegisterTool(ToolManagerCommands.TexturePaint, TexturePaintToolName, NewObject<UMeshTexturePaintingToolBuilder>());
+
+	SharedMeshToolData = NewObject<UMeshToolManager>(this);
+
+	UVertexAdapterClickToolBuilder* VertexClickToolBuilder = NewObject<UVertexAdapterClickToolBuilder>(this);
+	VertexClickToolBuilder->SharedMeshToolData = SharedMeshToolData;
+	RegisterTool(ToolManagerCommands.VertexSelect, VertexSelectToolName, VertexClickToolBuilder);
+
+	UTextureAdapterClickToolBuilder* TextureClickToolBuilder = NewObject<UTextureAdapterClickToolBuilder>(this);
+	TextureClickToolBuilder->SharedMeshToolData = SharedMeshToolData;
+	RegisterTool(ToolManagerCommands.TextureSelect, TextureSelectToolName, TextureClickToolBuilder);
+
+
+	UMeshColorPaintingToolBuilder* MeshColorPaintingToolBuilder = NewObject<UMeshColorPaintingToolBuilder>(this);
+	MeshColorPaintingToolBuilder->SharedMeshToolData = SharedMeshToolData;
+	RegisterTool(ToolManagerCommands.ColorPaint, ColorPaintToolName, MeshColorPaintingToolBuilder);
+
+	UMeshWeightPaintingToolBuilder* WeightPaintingToolBuilder = NewObject<UMeshWeightPaintingToolBuilder>(this);
+	WeightPaintingToolBuilder->SharedMeshToolData = SharedMeshToolData;
+	RegisterTool(ToolManagerCommands.WeightPaint, WeightPaintToolName, WeightPaintingToolBuilder);
+
+	UMeshTexturePaintingToolBuilder* TexturePaintingToolBuilder = NewObject<UMeshTexturePaintingToolBuilder>(this);
+	TexturePaintingToolBuilder->SharedMeshToolData = SharedMeshToolData;
+	RegisterTool(ToolManagerCommands.TexturePaint, TexturePaintToolName, TexturePaintingToolBuilder);
 	UpdateSelectedMeshes();
 
 
@@ -176,24 +193,17 @@ void UMeshPaintMode::Exit()
 	{
 		CommandList->UnmapAction(Action);
 	}
+
 	Super::Exit();
+
+	// Cleanup the SharedMeshToolData after the parent class exits, because the tools using it are closed out in UEdMode::Exit
+	SharedMeshToolData->ResetState();
+	SharedMeshToolData = nullptr;
 }
 
 void UMeshPaintMode::CreateToolkit()
 {
-	if (!Toolkit.IsValid())
-	{
-		FMeshPaintModeToolkit* PaintToolkit = new FMeshPaintModeToolkit;
-		Toolkit = MakeShareable(PaintToolkit);
-		Toolkit->Init(Owner->GetToolkitHost());
-
-		ToolsContext->OnToolNotificationMessage.AddSP(PaintToolkit, &FMeshPaintModeToolkit::SetActiveToolMessage);
-	}
-
-	// Register UI commands
-	BindCommands();
-
-	Super::CreateToolkit();
+	Toolkit = MakeShareable(new FMeshPaintModeToolkit);
 }
 
 void UMeshPaintMode::Tick(FEditorViewportClient* ViewportClient, float DeltaTime)
@@ -257,7 +267,9 @@ void UMeshPaintMode::BindCommands()
 	CommandList->MapAction(Commands.CycleToNextLOD, FExecuteAction::CreateUObject(this, &UMeshPaintMode::CycleMeshLODs, 1));
 	CommandList->MapAction(Commands.CycleToPreviousLOD, FExecuteAction::CreateUObject(this, &UMeshPaintMode::CycleMeshLODs, -1));
 
-	auto IsInValidPaintMode = [this]() -> bool { return (GetVertexToolProperties() != nullptr) && Cast<UMeshToolManager>(GetToolManager())->SelectionContainsValidAdapters(); };
+	auto IsInValidPaintMode = [this]() -> bool {
+		return (GetVertexToolProperties() != nullptr) && SharedMeshToolData && SharedMeshToolData->SelectionContainsValidAdapters();
+	};
 	CommandList->MapAction(Commands.Fill, 
 		FUIAction(FExecuteAction::CreateUObject(this, &UMeshPaintMode::FillWithVertexColor),
 		FCanExecuteAction::CreateLambda(IsInValidPaintMode)));
@@ -268,7 +280,9 @@ void UMeshPaintMode::BindCommands()
 			FCanExecuteAction::CreateUObject(this, &UMeshPaintMode::CanPropagateVertexColors)
 		));
 
-	auto IsAValidMeshComponentSelected = [this]() -> bool { return (GetSelectedComponents<UMeshComponent>().Num() == 1) && Cast<UMeshToolManager>(GetToolManager())->SelectionContainsValidAdapters(); };
+	auto IsAValidMeshComponentSelected = [this]() -> bool {
+		return (GetSelectedComponents<UMeshComponent>().Num() == 1) && SharedMeshToolData && SharedMeshToolData->SelectionContainsValidAdapters();
+	};
 	CommandList->MapAction(Commands.Import,
 		FUIAction(
 			FExecuteAction::CreateUObject(this, &UMeshPaintMode::ImportVertexColors),
@@ -348,9 +362,9 @@ void UMeshPaintMode::OnVertexPaintFinished()
 		}
 		else
 		{
-			if (UMeshToolManager* MeshToolManager = Cast<UMeshToolManager>(GetToolManager()))
+			if (SharedMeshToolData)
 			{
-				MeshToolManager->Refresh();
+				SharedMeshToolData->Refresh();
 			}
 		}
 	}
@@ -371,12 +385,12 @@ void UMeshPaintMode::ActorSelectionChangeNotify()
 
 void UMeshPaintMode::UpdateSelectedMeshes()
 {
-	if (UMeshToolManager* MeshToolManager = Cast<UMeshToolManager>(GetToolManager()))
+	if (SharedMeshToolData)
 	{
-		MeshToolManager->ResetState();
+		SharedMeshToolData->ResetState();
 		const TArray<UMeshComponent*> CurrentMeshComponents = GetSelectedComponents<UMeshComponent>();
-		MeshToolManager->AddSelectedMeshComponents(CurrentMeshComponents);
-		MeshToolManager->bNeedsRecache = true;
+		SharedMeshToolData->AddSelectedMeshComponents(CurrentMeshComponents);
+		SharedMeshToolData->bNeedsRecache = true;
 	}
 	bRecacheVertexDataSize = true;
 }
@@ -411,7 +425,7 @@ void UMeshPaintMode::FillWithVertexColor()
 		Component->Modify();
 		ComponentReregisterContext = MakeUnique<FComponentReregisterContext>(Component);
 
-		TSharedPtr<IMeshPaintComponentAdapter> MeshAdapter = Cast<UMeshToolManager>(GetToolManager())->GetAdapterForComponent(Component);
+		TSharedPtr<IMeshPaintComponentAdapter> MeshAdapter = SharedMeshToolData->GetAdapterForComponent(Component);
 		if (MeshAdapter)
 		{
 			MeshAdapter->PreEdit();
@@ -569,7 +583,7 @@ bool UMeshPaintMode::CanRemoveInstanceColors() const
 bool UMeshPaintMode::CanPasteInstanceVertexColors() const
 {
 	const TArray<UStaticMeshComponent*> StaticMeshComponents = GetSelectedComponents<UStaticMeshComponent>();
-	const TArray<FPerComponentVertexColorData> CopiedColorsByComponent = Cast<UMeshToolManager>(GetToolManager())->GetCopiedColorsByComponent();
+	const TArray<FPerComponentVertexColorData> CopiedColorsByComponent = SharedMeshToolData->GetCopiedColorsByComponent();
 	return UMeshPaintModeHelpers::CanPasteInstanceVertexColors(StaticMeshComponents, CopiedColorsByComponent);
 }
 
@@ -594,7 +608,7 @@ bool UMeshPaintMode::CanPropagateVertexColorsToLODs() const
 		bPaintOnSpecificLOD = ColorProperties ? ColorProperties->bPaintOnSpecificLOD : false;
 	}
 	// Can propagate when the mesh contains per-lod vertex colors or when we are not painting to a specific lod
-	const bool bSelectionContainsPerLODColors = Cast<UMeshToolManager>(GetToolManager())->SelectionContainsPerLODColors();
+	const bool bSelectionContainsPerLODColors = SharedMeshToolData->SelectionContainsPerLODColors();
 	return bSelectionContainsPerLODColors || !bPaintOnSpecificLOD;
 }
 
@@ -603,7 +617,7 @@ void UMeshPaintMode::CopyVertexColors()
 	const TArray<UStaticMeshComponent*> StaticMeshComponents = GetSelectedComponents<UStaticMeshComponent>();
 	TArray<FPerComponentVertexColorData> CopiedColorsByComponent;
 	UMeshPaintModeHelpers::CopyVertexColors(StaticMeshComponents, CopiedColorsByComponent);
-	Cast<UMeshToolManager>(GetToolManager())->SetCopiedColorsByComponent(CopiedColorsByComponent);
+	SharedMeshToolData->SetCopiedColorsByComponent(CopiedColorsByComponent);
 
 }
 
@@ -611,7 +625,7 @@ void UMeshPaintMode::PasteVertexColors()
 {
 	FScopedTransaction Transaction(LOCTEXT("LevelMeshPainter_TransactionPasteInstColors", "Pasting Per-Instance Vertex Colors"));
 	const TArray<UStaticMeshComponent*> StaticMeshComponents = GetSelectedComponents<UStaticMeshComponent>();
-	TArray<FPerComponentVertexColorData> CopiedColorsByComponent = Cast<UMeshToolManager>(GetToolManager())->GetCopiedColorsByComponent();
+	TArray<FPerComponentVertexColorData> CopiedColorsByComponent = SharedMeshToolData->GetCopiedColorsByComponent();
 	UMeshPaintModeHelpers::PasteVertexColors(StaticMeshComponents, CopiedColorsByComponent);
 	UpdateCachedVertexDataSize();
 }
@@ -656,9 +670,8 @@ void UMeshPaintMode::PropagateVertexColorsToLODs()
 {
 	//Only show the lost data warning if there is actually some data to lose
 	bool bAbortChange = false;
-	UMeshToolManager* MeshToolManager = Cast<UMeshToolManager>(GetToolManager());
-	TArray<UMeshComponent*> PaintableComponents = MeshToolManager->GetPaintableMeshComponents();
-	const bool bSelectionContainsPerLODColors = MeshToolManager->SelectionContainsPerLODColors();
+	TArray<UMeshComponent*> PaintableComponents = SharedMeshToolData->GetPaintableMeshComponents();
+	const bool bSelectionContainsPerLODColors = SharedMeshToolData->SelectionContainsPerLODColors();
 	if (bSelectionContainsPerLODColors)
 	{
 		//Warn the user they will lose custom painting data
@@ -679,7 +692,7 @@ void UMeshPaintMode::PropagateVertexColorsToLODs()
 		else
 		{
 			// Reset the state flag as we'll be removing all per-lod colors 
-			MeshToolManager->ClearSelectionLODColors();
+			SharedMeshToolData->ClearSelectionLODColors();
 			UMeshPaintModeHelpers::RemovePerLODColors(PaintableComponents);
 		}
 	}
@@ -694,13 +707,13 @@ void UMeshPaintMode::PropagateVertexColorsToLODs()
 	{
 		if (SelectedComponent)
 		{
-			TSharedPtr<IMeshPaintComponentAdapter> MeshAdapter = Cast<UMeshToolManager>(GetToolManager())->GetAdapterForComponent(SelectedComponent);
+			TSharedPtr<IMeshPaintComponentAdapter> MeshAdapter = SharedMeshToolData->GetAdapterForComponent(SelectedComponent);
 			UMeshPaintingToolset::ApplyVertexColorsToAllLODs(*MeshAdapter, SelectedComponent);
 			FComponentReregisterContext ReregisterContext(SelectedComponent);
 		}
 	}
 
-	MeshToolManager->Refresh();
+	SharedMeshToolData->Refresh();
 }
 
 template<typename ComponentClass>
@@ -746,6 +759,10 @@ template TArray<UStaticMeshComponent*> UMeshPaintMode::GetSelectedComponents<USt
 template TArray<USkeletalMeshComponent*> UMeshPaintMode::GetSelectedComponents<USkeletalMeshComponent>() const;
 template TArray<UMeshComponent*> UMeshPaintMode::GetSelectedComponents<UMeshComponent>() const;
 
+TWeakObjectPtr<UMeshToolManager> UMeshPaintMode::GetMeshToolData() const
+{
+	return SharedMeshToolData;
+}
 
 
 
@@ -754,9 +771,9 @@ void UMeshPaintMode::UpdateCachedVertexDataSize()
 	CachedVertexDataSize = 0;
 
 	const bool bInstance = true;
-	if (UMeshToolManager* MeshToolManager = Cast<UMeshToolManager>(GetToolManager()))
+	if (SharedMeshToolData)
 	{
-		for (UMeshComponent* SelectedComponent : MeshToolManager->GetPaintableMeshComponents())
+		for (UMeshComponent* SelectedComponent : SharedMeshToolData->GetPaintableMeshComponents())
 		{
 			int32 NumLODs = UMeshPaintingToolset::GetNumberOfLODs(SelectedComponent);
 			for (int32 LODIndex = 0; LODIndex < NumLODs; ++LODIndex)
diff --git a/Engine/Plugins/MeshPainting/Source/MeshPaintEditorMode/Private/MeshPaintMode.h b/Engine/Plugins/MeshPainting/Source/MeshPaintEditorMode/Private/MeshPaintMode.h
--- a/Engine/Plugins/MeshPainting/Source/MeshPaintEditorMode/Private/MeshPaintMode.h
+++ b/Engine/Plugins/MeshPainting/Source/MeshPaintEditorMode/Private/MeshPaintMode.h
@@ -22,6 +22,7 @@ class UMeshTexturePaintingToolProperties;
 class UMeshPaintModeSettings;
 class IMeshPaintComponentAdapter;
 class UMeshComponent;
+class UMeshToolManager;
 
 /**
  * Mesh paint Mode.  Extends editor viewports with the ability to paint data on meshes
@@ -62,10 +63,12 @@ public:
 		return CachedVertexDataSize;
 	}
 
+	TWeakObjectPtr<UMeshToolManager> GetMeshToolData() const;
+
 protected:
 
 	/** Binds UI commands to actions for the mesh paint mode */
-	virtual void BindCommands();
+	virtual void BindCommands() override;
 
 	// UEdMode interface
 	virtual void OnToolStarted(UInteractiveToolManager* Manager, UInteractiveTool* Tool) override;
@@ -100,9 +103,14 @@ protected:
 	void CommitAllPaintedTextures();
 	int32 GetNumberOfPendingPaintChanges();
 	void OnVertexPaintFinished();
+
 protected:
 	UPROPERTY(Transient)
 	UMeshPaintModeSettings* ModeSettings;
+
+	UPROPERTY(Transient)
+	UMeshToolManager* SharedMeshToolData;
+
 	// End vertex paint state
 	FGetSelectedMeshComponents MeshComponentDelegate;
 	uint32 CachedVertexDataSize;
diff --git a/Engine/Plugins/MeshPainting/Source/MeshPaintEditorMode/Private/MeshPaintModeToolkit.cpp b/Engine/Plugins/MeshPainting/Source/MeshPaintEditorMode/Private/MeshPaintModeToolkit.cpp
--- a/Engine/Plugins/MeshPainting/Source/MeshPaintEditorMode/Private/MeshPaintModeToolkit.cpp
+++ b/Engine/Plugins/MeshPainting/Source/MeshPaintEditorMode/Private/MeshPaintModeToolkit.cpp
@@ -5,11 +5,16 @@
 
 #define LOCTEXT_NAMESPACE "MeshPaintModeToolkit"
 
-void FMeshPaintModeToolkit::Init(const TSharedPtr<IToolkitHost>& InitToolkitHost)
+FMeshPaintModeToolkit::~FMeshPaintModeToolkit()
 {
+	GetScriptableEditorMode()->GetInteractiveToolsContext()->OnToolNotificationMessage.RemoveAll(this);
+}
+
+void FMeshPaintModeToolkit::Init(const TSharedPtr<IToolkitHost>& InitToolkitHost, TWeakObjectPtr<UEdMode> InOwningMode)
+{
+	FModeToolkit::Init(InitToolkitHost, InOwningMode);
 
-	FModeToolkit::Init(InitToolkitHost);
-	// ToolsContextClass = UMeshToolsContext::StaticClass();
+	GetScriptableEditorMode()->GetInteractiveToolsContext()->OnToolNotificationMessage.AddSP(this, &FMeshPaintModeToolkit::SetActiveToolMessage);
 }
 
 FName FMeshPaintModeToolkit::GetToolkitFName() const
@@ -23,12 +28,6 @@ FText FMeshPaintModeToolkit::GetBaseToolkitName() const
 	return LOCTEXT( "ToolkitName", "Mesh Paint Mode" );
 }
 
-
-class UEdMode* FMeshPaintModeToolkit::GetScriptableEditorMode() const
-{
-	return GLevelEditorModeTools().GetActiveScriptableMode( "MeshPaintMode" );
-}
-
 void FMeshPaintModeToolkit::GetToolPaletteNames(TArray<FName>& PaletteNames) const
 {
 	PaletteNames.Add(UMeshPaintMode::MeshPaintMode_Color);
diff --git a/Engine/Plugins/MeshPainting/Source/MeshPaintEditorMode/Private/MeshPaintModeToolkit.h b/Engine/Plugins/MeshPainting/Source/MeshPaintEditorMode/Private/MeshPaintModeToolkit.h
--- a/Engine/Plugins/MeshPainting/Source/MeshPaintEditorMode/Private/MeshPaintModeToolkit.h
+++ b/Engine/Plugins/MeshPainting/Source/MeshPaintEditorMode/Private/MeshPaintModeToolkit.h
@@ -15,13 +15,14 @@ public:
 	FMeshPaintModeToolkit(  )
 	{}
 
+	~FMeshPaintModeToolkit();
+
 	/** FModeToolkit interface */
-	virtual void Init(const TSharedPtr<IToolkitHost>& InitToolkitHost) override;
+	virtual void Init(const TSharedPtr<IToolkitHost>& InitToolkitHost, TWeakObjectPtr<UEdMode> InOwningMode) override;
 	
 	// IToolkit overrides
 	virtual FName GetToolkitFName() const override;
 	virtual FText GetBaseToolkitName() const override;
-	virtual class UEdMode* GetScriptableEditorMode() const override;
 	virtual void GetToolPaletteNames(TArray<FName>& PaletteNames) const override;
 	virtual FText GetToolPaletteDisplayName(FName Palette) const override;
 
diff --git a/Engine/Plugins/MeshPainting/Source/MeshPaintingToolset/Private/MeshPaintHelpers.cpp b/Engine/Plugins/MeshPainting/Source/MeshPaintingToolset/Private/MeshPaintHelpers.cpp
--- a/Engine/Plugins/MeshPainting/Source/MeshPaintingToolset/Private/MeshPaintHelpers.cpp
+++ b/Engine/Plugins/MeshPainting/Source/MeshPaintingToolset/Private/MeshPaintHelpers.cpp
@@ -1395,14 +1395,6 @@ UMeshToolManager::UMeshToolManager()
 	bNeedsRecache = true;
 }
 
-void UMeshToolManager::Shutdown()
-{
-	UInteractiveToolManager::Shutdown();
-
-	ComponentToAdapterMap.Empty();
-	SelectedMeshComponents.Empty();
-}
-
 TMap<UMeshComponent*, TSharedPtr<IMeshPaintComponentAdapter>> UMeshToolManager::GetComponentToAdapterMap() const
 {
 	return ComponentToAdapterMap;
@@ -1559,9 +1551,3 @@ int32 UMeshToolManager::GetMaxUVIndexToPaint() const
 
 	return 0;
 }
-
-UMeshToolsContext::UMeshToolsContext()
-	:Super()
-{
-	ToolManagerClass = UMeshToolManager::StaticClass();
-}
diff --git a/Engine/Plugins/MeshPainting/Source/MeshPaintingToolset/Private/MeshPaintInteractions.cpp b/Engine/Plugins/MeshPainting/Source/MeshPaintingToolset/Private/MeshPaintInteractions.cpp
--- a/Engine/Plugins/MeshPainting/Source/MeshPaintingToolset/Private/MeshPaintInteractions.cpp
+++ b/Engine/Plugins/MeshPainting/Source/MeshPaintingToolset/Private/MeshPaintInteractions.cpp
@@ -19,51 +19,52 @@
 
 FInputRayHit UMeshPaintSelectionMechanic::IsHitByClick(const FInputDeviceRay& ClickPos)
 {
-	if (UMeshToolManager* MeshToolManager = Cast<UMeshToolManager>(GetParentTool()->GetToolManager()))
+	IMeshPaintSelectionInterface* Interface = Cast<IMeshPaintSelectionInterface>(GetParentTool());
+	if (!bAddToSelectionSet || !Interface->AllowsMultiselect())
 	{
-		IMeshPaintSelectionInterface* Interface = Cast<IMeshPaintSelectionInterface>(GetParentTool());
-		if (!bAddToSelectionSet || !Interface->AllowsMultiselect())
-		{
-			CachedClickedComponents.Empty();
-			CachedClickedActors.Empty();
-		}
-		return FindClickedComponentsAndCacheAdapters(ClickPos, MeshToolManager) ? FInputRayHit(0.0f) : FInputRayHit();
+		CachedClickedComponents.Empty();
+		CachedClickedActors.Empty();
 	}
-	return FInputRayHit();
+	return FindClickedComponentsAndCacheAdapters(ClickPos) ? FInputRayHit(0.0f) : FInputRayHit();
 }
 
 void UMeshPaintSelectionMechanic::OnClicked(const FInputDeviceRay& ClickPos)
 {
-	if (UMeshToolManager* MeshToolManager = Cast<UMeshToolManager>(GetParentTool()->GetToolManager()))
+	if (SharedMeshToolData.IsValid())
 	{
 		for (UMeshComponent* MeshComponent : CachedClickedComponents)
 		{
-			TSharedPtr<IMeshPaintComponentAdapter> MeshAdapter = MeshToolManager->GetAdapterForComponent(MeshComponent);
+			TSharedPtr<IMeshPaintComponentAdapter> MeshAdapter = SharedMeshToolData->GetAdapterForComponent(MeshComponent);
 			IMeshPaintSelectionInterface* Interface = Cast<IMeshPaintSelectionInterface>(GetParentTool());
 			if (MeshComponent && MeshComponent->IsVisible() && MeshAdapter.IsValid() && MeshAdapter->IsValid() && Interface->IsMeshAdapterSupported(MeshAdapter))
 			{
-				MeshToolManager->AddPaintableMeshComponent(MeshComponent);
+				SharedMeshToolData->AddPaintableMeshComponent(MeshComponent);
 				MeshAdapter->OnAdded();
 			}
 
-			MeshToolManager->BeginUndoTransaction(LOCTEXT("MeshSelection", "Select Mesh"));
+			GetParentTool()->GetToolManager()->BeginUndoTransaction(LOCTEXT("MeshSelection", "Select Mesh"));
 
 
 			FSelectedOjectsChangeList NewSelection;
 			// TODO add CTRL handling
 			NewSelection.ModificationType = bAddToSelectionSet ? ESelectedObjectsModificationType::Add : ESelectedObjectsModificationType::Replace;
 			NewSelection.Actors.Append(CachedClickedActors);
-			MeshToolManager->RequestSelectionChange(NewSelection);
-			MeshToolManager->EndUndoTransaction();
+			GetParentTool()->GetToolManager()->RequestSelectionChange(NewSelection);
+			GetParentTool()->GetToolManager()->EndUndoTransaction();
 		}
 	}
 }
 
-bool UMeshPaintSelectionMechanic::FindClickedComponentsAndCacheAdapters(const FInputDeviceRay& ClickPos, class UMeshToolManager* MeshToolManager)
+bool UMeshPaintSelectionMechanic::FindClickedComponentsAndCacheAdapters(const FInputDeviceRay& ClickPos)
 {
 	bool bFoundValidComponents = false;
 #if WITH_EDITOR
-	if (HHitProxy* HitProxy = MeshToolManager->GetContextQueriesAPI()->GetHitProxy(ClickPos.ScreenPosition.X, ClickPos.ScreenPosition.Y))
+	if (!SharedMeshToolData.IsValid())
+	{
+		return bFoundValidComponents;
+	}
+
+	if (HHitProxy* HitProxy = GetParentTool()->GetToolManager()->GetContextQueriesAPI()->GetHitProxy(ClickPos.ScreenPosition.X, ClickPos.ScreenPosition.Y))
 	{
 		if (HitProxy->IsA(HActor::StaticGetType()))
 		{
@@ -77,7 +78,7 @@ bool UMeshPaintSelectionMechanic::FindClickedComponentsAndCacheAdapters(const FI
 				IMeshPaintSelectionInterface* Interface = Cast<IMeshPaintSelectionInterface>(GetParentTool());
 				if (MeshAdapter.IsValid() && Interface->IsMeshAdapterSupported(MeshAdapter))
 				{
-					MeshToolManager->AddToComponentToAdapterMap(MeshComponent, MeshAdapter);
+					SharedMeshToolData->AddToComponentToAdapterMap(MeshComponent, MeshAdapter);
 					CachedClickedComponents.AddUnique(MeshComponent);
 					CachedClickedActors.AddUnique(Cast<AActor>(MeshComponent->GetOuter()));
 					bFoundValidComponents = true;
@@ -88,4 +89,10 @@ bool UMeshPaintSelectionMechanic::FindClickedComponentsAndCacheAdapters(const FI
 #endif
 	return bFoundValidComponents;
 }
+
+void UMeshPaintSelectionMechanic::SetMeshToolData(TWeakObjectPtr<UMeshToolManager> InMeshToolData)
+{
+	SharedMeshToolData = InMeshToolData;
+}
+
 #undef LOCTEXT_NAMESPACE
\ No newline at end of file
diff --git a/Engine/Plugins/MeshPainting/Source/MeshPaintingToolset/Private/MeshSelect.cpp b/Engine/Plugins/MeshPainting/Source/MeshPaintingToolset/Private/MeshSelect.cpp
--- a/Engine/Plugins/MeshPainting/Source/MeshPaintingToolset/Private/MeshSelect.cpp
+++ b/Engine/Plugins/MeshPainting/Source/MeshPaintingToolset/Private/MeshSelect.cpp
@@ -22,7 +22,7 @@ bool UVertexAdapterClickToolBuilder::CanBuildTool(const FToolBuilderState& Scene
 UInteractiveTool* UVertexAdapterClickToolBuilder::BuildTool(const FToolBuilderState& SceneState) const
 {
 	UVertexAdapterClickTool* NewTool = NewObject<UVertexAdapterClickTool>(SceneState.ToolManager);
-
+	NewTool->SetMeshToolData(SharedMeshToolData);
 	return NewTool;
 }
 
@@ -34,6 +34,7 @@ bool UTextureAdapterClickToolBuilder::CanBuildTool(const FToolBuilderState& Scen
 UInteractiveTool* UTextureAdapterClickToolBuilder::BuildTool(const FToolBuilderState& SceneState) const
 {
 	UTextureAdapterClickTool* NewTool = NewObject<UTextureAdapterClickTool>(SceneState.ToolManager);
+	NewTool->SetMeshToolData(SharedMeshToolData);
 
 	return NewTool;
 }
@@ -60,6 +61,7 @@ void UMeshClickTool::Setup()
 	// Set up selection mechanic to select valid meshes
 	SelectionMechanic = NewObject<UMeshPaintSelectionMechanic>(this);
 	SelectionMechanic->Setup(this);
+	SelectionMechanic->SetMeshToolData(SharedMeshToolData);
 }
 
 void UMeshClickTool::OnUpdateModifierState(int ModifierID, bool bIsOn)
@@ -80,6 +82,10 @@ void UMeshClickTool::OnClicked(const FInputDeviceRay& ClickPos)
 	SelectionMechanic->OnClicked(ClickPos);
 }
 
+void UMeshClickTool::SetMeshToolData(TWeakObjectPtr<UMeshToolManager> InMeshToolData)
+{
+	SharedMeshToolData = InMeshToolData;
+}
 
 
 UVertexAdapterClickTool::UVertexAdapterClickTool()
diff --git a/Engine/Plugins/MeshPainting/Source/MeshPaintingToolset/Private/MeshTexturePaintingTool.cpp b/Engine/Plugins/MeshPainting/Source/MeshPaintingToolset/Private/MeshTexturePaintingTool.cpp
--- a/Engine/Plugins/MeshPainting/Source/MeshPaintingToolset/Private/MeshTexturePaintingTool.cpp
+++ b/Engine/Plugins/MeshPainting/Source/MeshPaintingToolset/Private/MeshTexturePaintingTool.cpp
@@ -31,9 +31,8 @@
 bool UMeshTexturePaintingToolBuilder::CanBuildTool(const FToolBuilderState& SceneState) const
 {
 	UActorComponent* ActorComponent = ToolBuilderUtil::FindFirstComponent(SceneState, UMeshPaintingToolset::HasPaintableMesh);
-	UMeshToolManager* MeshToolManager = Cast<UMeshToolManager>(SceneState.ToolManager);
 	TArray<FPaintableTexture> PaintableTextures;
-	if (ActorComponent && MeshToolManager)
+	if (ActorComponent && SharedMeshToolData.IsValid())
 	{
 		UMeshComponent* MeshComponent = Cast<UMeshComponent>(ActorComponent);
 		if (MeshComponent)
@@ -49,6 +48,7 @@ bool UMeshTexturePaintingToolBuilder::CanBuildTool(const FToolBuilderState& Scen
 UInteractiveTool* UMeshTexturePaintingToolBuilder::BuildTool(const FToolBuilderState& SceneState) const
 {
 	UMeshTexturePaintingTool* NewTool = NewObject<UMeshTexturePaintingTool>(SceneState.ToolManager);
+	NewTool->SetMeshToolData(SharedMeshToolData);
 	return NewTool;
 }
 
@@ -109,10 +109,9 @@ void UMeshTexturePaintingTool::Shutdown(EToolShutdownType ShutdownType)
 
 	// Remove any existing texture targets
 	TexturePaintTargetList.Empty();
-	UMeshToolManager* MeshToolManager = Cast<UMeshToolManager>(GetToolManager());
-	if (MeshToolManager)
+	if (SharedMeshToolData.IsValid())
 	{
-		MeshToolManager->Refresh();
+		SharedMeshToolData->Refresh();
 	}
 
 	BrushProperties->SaveProperties(this);
@@ -121,11 +120,11 @@ void UMeshTexturePaintingTool::Shutdown(EToolShutdownType ShutdownType)
 
 void UMeshTexturePaintingTool::Render(IToolsContextRenderAPI* RenderAPI)
 {
-	UMeshToolManager* MeshToolManager = Cast<UMeshToolManager>(GetToolManager());
-		Super::Render(RenderAPI);
-		FToolDataVisualizer Draw;
-		Draw.BeginFrame(RenderAPI);
-	if (MeshToolManager && LastBestHitResult.Component != nullptr)
+	Super::Render(RenderAPI);
+	FToolDataVisualizer Draw;
+	Draw.BeginFrame(RenderAPI);
+
+	if (SharedMeshToolData.IsValid() && LastBestHitResult.Component != nullptr)
 	{
 		BrushStampIndicator->bDrawIndicatorLines = true;
 		static float WidgetLineThickness = 1.0f;
@@ -139,9 +138,9 @@ void UMeshTexturePaintingTool::Render(IToolsContextRenderAPI* RenderAPI)
 		const FVector NormalLineEnd(LastBestHitResult.Location + LastBestHitResult.Normal * NormalLineSize);
 		Draw.DrawLine(FVector(LastBestHitResult.Location), NormalLineEnd, NormalLineColor, WidgetLineThickness);
 
-		for (UMeshComponent* CurrentComponent : MeshToolManager->GetPaintableMeshComponents())
+		for (UMeshComponent* CurrentComponent : SharedMeshToolData->GetPaintableMeshComponents())
 		{
-			TSharedPtr<IMeshPaintComponentAdapter> MeshAdapter = MeshToolManager->GetAdapterForComponent(Cast<UMeshComponent>(CurrentComponent));
+			TSharedPtr<IMeshPaintComponentAdapter> MeshAdapter = SharedMeshToolData->GetAdapterForComponent(Cast<UMeshComponent>(CurrentComponent));
 
 			if (MeshAdapter->IsValid() && MeshAdapter->SupportsVertexPaint())
 			{
@@ -168,9 +167,9 @@ void UMeshTexturePaintingTool::Render(IToolsContextRenderAPI* RenderAPI)
 
 void UMeshTexturePaintingTool::OnTick(float DeltaTime)
 {
-	if (UMeshToolManager* MeshToolManager = Cast<UMeshToolManager>(GetToolManager()))
+	if (SharedMeshToolData.IsValid())
 	{
-		if (MeshToolManager->bNeedsRecache)
+		if (SharedMeshToolData->bNeedsRecache)
 		{
 			CacheSelectionData();
 			CacheTexturePaintData();
@@ -200,15 +199,14 @@ void UMeshTexturePaintingTool::OnPropertyModified(UObject* PropertySet, FPropert
 
 double UMeshTexturePaintingTool::EstimateMaximumTargetDimension()
 {
-	UMeshToolManager* MeshToolManager = Cast<UMeshToolManager>(GetToolManager());
 	bool bFoundComponentToUse = false;
 	FBoxSphereBounds Bounds = FBoxSphereBounds(0.0);
-	if (MeshToolManager)
+	if (SharedMeshToolData.IsValid())
 	{
 		bool bFirstItem = true;
 
 		FBoxSphereBounds Extents;
-		for (UMeshComponent* SelectedComponent : MeshToolManager->GetSelectedMeshComponents())
+		for (UMeshComponent* SelectedComponent : SharedMeshToolData->GetSelectedMeshComponents())
 		{
 			if (bFirstItem)
 			{
@@ -258,30 +256,30 @@ bool UMeshTexturePaintingTool::Paint(const TArrayView<TPair<FVector, FVector>>&
 
 void UMeshTexturePaintingTool::CacheSelectionData()
 {
-	if (UMeshToolManager* MeshToolManager = Cast<UMeshToolManager>(GetToolManager()))
+	if (SharedMeshToolData.IsValid())
 	{
-		MeshToolManager->ClearPaintableMeshComponents();
+		SharedMeshToolData->ClearPaintableMeshComponents();
 
 		//Determine LOD level to use for painting(can only paint on LODs in vertex mode)
 		const int32 PaintLODIndex = 0;
 		//Determine UV channel to use while painting textures
 		const int32 UVChannel = 0;
 
-		MeshToolManager->CacheSelectionData(PaintLODIndex, UVChannel);
+		SharedMeshToolData->CacheSelectionData(PaintLODIndex, UVChannel);
 	}
 }
 
 void UMeshTexturePaintingTool::CacheTexturePaintData()
 {
-	if (UMeshToolManager* MeshToolManager = Cast<UMeshToolManager>(GetToolManager()))
+	if (SharedMeshToolData.IsValid())
 	{
-		TArray<UMeshComponent*> PaintableComponents = MeshToolManager->GetPaintableMeshComponents();
+		TArray<UMeshComponent*> PaintableComponents = SharedMeshToolData->GetPaintableMeshComponents();
 		
 		PaintableTextures.Empty();
 		if (PaintableComponents.Num() == 1 && PaintableComponents[0])
 		{
 			const UMeshComponent* Component = PaintableComponents[0];
-			TSharedPtr<IMeshPaintComponentAdapter> Adapter = MeshToolManager->GetAdapterForComponent(Component);
+			TSharedPtr<IMeshPaintComponentAdapter> Adapter = SharedMeshToolData->GetAdapterForComponent(Component);
 			UTexturePaintToolset::RetrieveTexturesForComponent(Component, Adapter.Get(), PaintableTextures);
 		}
 
@@ -528,9 +526,9 @@ void UMeshTexturePaintingTool::OnEndDrag(const FRay& Ray)
 bool UMeshTexturePaintingTool::HitTest(const FRay& Ray, FHitResult& OutHit)
 {
 	bool bUsed = false;
-	if (UMeshToolManager* MeshToolManager = Cast<UMeshToolManager>(GetToolManager()))
+	if (SharedMeshToolData.IsValid())
 	{
-		MeshToolManager->FindHitResult(Ray, OutHit);
+		SharedMeshToolData->FindHitResult(Ray, OutHit);
 		LastBestHitResult = OutHit;
 		bUsed = OutHit.bBlockingHit;
 	}
@@ -1305,4 +1303,9 @@ void UMeshTexturePaintingTool::PaintTextureChanged(const FAssetData& AssetData)
 	}
 }
 
+void UMeshTexturePaintingTool::SetMeshToolData(TWeakObjectPtr<UMeshToolManager> InToolData)
+{
+	SharedMeshToolData = InToolData;
+}
+
 #undef LOCTEXT_NAMESPACE
diff --git a/Engine/Plugins/MeshPainting/Source/MeshPaintingToolset/Private/MeshVertexPaintingTool.cpp b/Engine/Plugins/MeshPainting/Source/MeshPaintingToolset/Private/MeshVertexPaintingTool.cpp
--- a/Engine/Plugins/MeshPainting/Source/MeshPaintingToolset/Private/MeshVertexPaintingTool.cpp
+++ b/Engine/Plugins/MeshPainting/Source/MeshPaintingToolset/Private/MeshVertexPaintingTool.cpp
@@ -29,6 +29,7 @@ bool UMeshColorPaintingToolBuilder::CanBuildTool(const FToolBuilderState& SceneS
 UInteractiveTool* UMeshColorPaintingToolBuilder::BuildTool(const FToolBuilderState& SceneState) const
 {
 	UMeshColorPaintingTool* NewTool = NewObject<UMeshColorPaintingTool>(SceneState.ToolManager);
+	NewTool->SetMeshToolData(SharedMeshToolData);
 	return NewTool;
 }
 
@@ -40,6 +41,7 @@ bool UMeshWeightPaintingToolBuilder::CanBuildTool(const FToolBuilderState& Scene
 UInteractiveTool* UMeshWeightPaintingToolBuilder::BuildTool(const FToolBuilderState& SceneState) const
 {
 	UMeshWeightPaintingTool* NewTool = NewObject<UMeshWeightPaintingTool>(SceneState.ToolManager);
+	NewTool->SetMeshToolData(SharedMeshToolData);
 	return NewTool;
 }
 
@@ -67,6 +69,11 @@ bool UMeshVertexPaintingTool::IsMeshAdapterSupported(TSharedPtr<IMeshPaintCompon
 	return MeshAdapter.IsValid() ? MeshAdapter->SupportsVertexPaint() : false;
 }
 
+void UMeshVertexPaintingTool::SetMeshToolData(TWeakObjectPtr<UMeshToolManager> InMeshToolData)
+{
+	SharedMeshToolData = InMeshToolData;
+}
+
 void UMeshVertexPaintingTool::Setup()
 {
 	Super::Setup();
@@ -78,7 +85,7 @@ void UMeshVertexPaintingTool::Setup()
 	// Set up selection mechanic to find and select edges
 	SelectionMechanic = NewObject<UMeshPaintSelectionMechanic>(this);
 	SelectionMechanic->Setup(this);
-	
+	SelectionMechanic->SetMeshToolData(SharedMeshToolData);
 }
 
 
@@ -86,21 +93,19 @@ void UMeshVertexPaintingTool::Shutdown(EToolShutdownType ShutdownType)
 {
 	FinishPainting();
 	BrushProperties->SaveProperties(this);
-	UMeshToolManager* MeshToolManager = Cast<UMeshToolManager>(GetToolManager());
-	if (MeshToolManager)
+	if (SharedMeshToolData.IsValid())
 	{
-		MeshToolManager->Refresh();
+		SharedMeshToolData->Refresh();
 	}
 	Super::Shutdown(ShutdownType);
 }
 
 void UMeshVertexPaintingTool::Render(IToolsContextRenderAPI* RenderAPI)
 {
-	UMeshToolManager* MeshToolManager = Cast<UMeshToolManager>(GetToolManager());
 	Super::Render(RenderAPI);
 	FToolDataVisualizer Draw;
 	Draw.BeginFrame(RenderAPI);
-	if (MeshToolManager && LastBestHitResult.Component != nullptr && MeshToolManager->GetPaintableMeshComponents().Num() > 0)
+	if (SharedMeshToolData.IsValid() && LastBestHitResult.Component != nullptr && SharedMeshToolData->GetPaintableMeshComponents().Num() > 0)
 	{
 		BrushStampIndicator->bDrawIndicatorLines = true;
 		static float WidgetLineThickness = 1.0f;
@@ -115,9 +120,9 @@ void UMeshVertexPaintingTool::Render(IToolsContextRenderAPI* RenderAPI)
 		const FVector NormalLineEnd(LastBestHitResult.Location + LastBestHitResult.Normal * NormalLineSize);
 		Draw.DrawLine(FVector(LastBestHitResult.Location), NormalLineEnd, NormalLineColor, WidgetLineThickness);
 
-		for (UMeshComponent* CurrentComponent : MeshToolManager->GetPaintableMeshComponents())
+		for (UMeshComponent* CurrentComponent : SharedMeshToolData->GetPaintableMeshComponents())
 		{
-			TSharedPtr<IMeshPaintComponentAdapter> MeshAdapter = MeshToolManager->GetAdapterForComponent(Cast<UMeshComponent>(CurrentComponent));
+			TSharedPtr<IMeshPaintComponentAdapter> MeshAdapter = SharedMeshToolData->GetAdapterForComponent(Cast<UMeshComponent>(CurrentComponent));
 
 			if (MeshAdapter->IsValid() && MeshAdapter->SupportsVertexPaint())
 			{
@@ -157,9 +162,9 @@ void UMeshVertexPaintingTool::Render(IToolsContextRenderAPI* RenderAPI)
 
 void UMeshVertexPaintingTool::OnTick(float DeltaTime)
 {
-	if (UMeshToolManager* MeshToolManager = Cast<UMeshToolManager>(GetToolManager()))
+	if (SharedMeshToolData.IsValid())
 	{
-		if (MeshToolManager->bNeedsRecache)
+		if (SharedMeshToolData->bNeedsRecache)
 		{
 			CacheSelectionData();
 			bDoRestoreRenTargets = true;
@@ -188,15 +193,14 @@ void UMeshVertexPaintingTool::OnPropertyModified(UObject* PropertySet, FProperty
 
 double UMeshVertexPaintingTool::EstimateMaximumTargetDimension()
 {
-	UMeshToolManager* MeshToolManager = Cast<UMeshToolManager>(GetToolManager());
 	bool bFoundComponentToUse = false;
 	FBoxSphereBounds Bounds = FBoxSphereBounds(0.0);
-	if (MeshToolManager)
+	if (SharedMeshToolData.IsValid())
 	{
 		bool bFirstItem = true;
 
 		FBoxSphereBounds Extents;
-		for (UMeshComponent* SelectedComponent : MeshToolManager->GetSelectedMeshComponents())
+		for (UMeshComponent* SelectedComponent : SharedMeshToolData->GetSelectedMeshComponents())
 		{
 			if (bFirstItem)
 			{
@@ -260,7 +264,6 @@ bool UMeshVertexPaintingTool::PaintInternal(const TArrayView<TPair<FVector, FVec
 	const float InStrengthScale = PaintStrength;
 
 	bool bPaintApplied = false;
-	UMeshToolManager* MeshToolManager = Cast<UMeshToolManager>(GetToolManager());
 	// Fire out a ray to see if there is a *selected* component under the mouse cursor that can be painted.
 	for (int32 i = 0; i < Rays.Num(); ++i)
 	{
@@ -271,9 +274,9 @@ bool UMeshVertexPaintingTool::PaintInternal(const TArrayView<TPair<FVector, FVec
 		const FVector TraceStart(RayOrigin);
 		const FVector TraceEnd(RayOrigin + RayDirection * HALF_WORLD_MAX);
 
-		for (UMeshComponent* MeshComponent : MeshToolManager->GetPaintableMeshComponents())
+		for (UMeshComponent* MeshComponent : SharedMeshToolData->GetPaintableMeshComponents())
 		{
-			TSharedPtr<IMeshPaintComponentAdapter> MeshAdapter = MeshToolManager->GetAdapterForComponent(MeshComponent);
+			TSharedPtr<IMeshPaintComponentAdapter> MeshAdapter = SharedMeshToolData->GetAdapterForComponent(MeshComponent);
 
 			// Ray trace
 			FHitResult TraceHitResult(1.0f);
@@ -295,11 +298,11 @@ bool UMeshVertexPaintingTool::PaintInternal(const TArrayView<TPair<FVector, FVec
 			FBox BrushBounds = FBox::BuildAABB(BestTraceResult.Location, FVector(BrushRadius * 1.25f, BrushRadius * 1.25f, BrushRadius * 1.25f));
 
 			// Vertex paint mode, so we want all valid components overlapping the brush hit location
-			for (auto TestComponent : MeshToolManager->GetPaintableMeshComponents())
+			for (auto TestComponent : SharedMeshToolData->GetPaintableMeshComponents())
 			{
 				const FBox ComponentBounds = TestComponent->Bounds.GetBox();
 
-				if (MeshToolManager->GetComponentToAdapterMap().Contains(TestComponent) && ComponentBounds.Intersect(BrushBounds))
+				if (SharedMeshToolData->GetComponentToAdapterMap().Contains(TestComponent) && ComponentBounds.Intersect(BrushBounds))
 				{
 					// OK, this mesh potentially overlaps the brush!
 					HoveredComponents.FindOrAdd(TestComponent).Add(i);
@@ -361,7 +364,7 @@ bool UMeshVertexPaintingTool::PaintInternal(const TArrayView<TPair<FVector, FVec
 			UMeshComponent* HoveredComponent = Entry.Key;
 			TArray<int32>& PaintRayResultIds = Entry.Value;
 
-			IMeshPaintComponentAdapter* MeshAdapter = MeshToolManager->GetAdapterForComponent(HoveredComponent).Get();
+			IMeshPaintComponentAdapter* MeshAdapter = SharedMeshToolData->GetAdapterForComponent(HoveredComponent).Get();
 			if (!ensure(MeshAdapter))
 			{
 				continue;
@@ -512,9 +515,9 @@ void UMeshVertexPaintingTool::OnEndDrag(const FRay& Ray)
 bool UMeshVertexPaintingTool::HitTest(const FRay& Ray, FHitResult& OutHit)
 {
 	bool bUsed = false;
-	if (UMeshToolManager* MeshToolManager = Cast<UMeshToolManager>(GetToolManager()))
+	if (SharedMeshToolData.IsValid())
 	{
-		MeshToolManager->FindHitResult(Ray, OutHit);
+		SharedMeshToolData->FindHitResult(Ray, OutHit);
 		LastBestHitResult = OutHit;
 		bUsed = OutHit.bBlockingHit;
 	}
@@ -567,9 +570,9 @@ void UMeshColorPaintingTool::Shutdown(EToolShutdownType ShutdownType)
 void UMeshColorPaintingTool::CacheSelectionData()
 {
 	Super::CacheSelectionData();
-	if (UMeshToolManager* MeshToolManager = Cast<UMeshToolManager>(GetToolManager()))
+	if (SharedMeshToolData.IsValid())
 	{
-		MeshToolManager->ClearPaintableMeshComponents();
+		SharedMeshToolData->ClearPaintableMeshComponents();
 		// Update(cached) Paint LOD level if necessary
 		ColorProperties->LODIndex = FMath::Min<int32>(ColorProperties->LODIndex, GetMaxLODIndexToPaint());
 		CachedLODIndex = ColorProperties->LODIndex;
@@ -579,7 +582,7 @@ void UMeshColorPaintingTool::CacheSelectionData()
 		//Determine UV channel to use while painting textures
 		const int32 UVChannel = 0;
 
-		MeshToolManager->CacheSelectionData(PaintLODIndex, UVChannel);
+		SharedMeshToolData->CacheSelectionData(PaintLODIndex, UVChannel);
 	}
 }
 
@@ -597,9 +600,9 @@ int32 UMeshColorPaintingTool::GetMaxLODIndexToPaint() const
 	//The maximum LOD we can paint is decide by the lowest number of LOD in the selection
 	int32 LODMin = TNumericLimits<int32>::Max();
 
-	if (UMeshToolManager* MeshToolManager = Cast<UMeshToolManager>(GetToolManager()))
+	if (SharedMeshToolData.IsValid())
 	{
-		TArray<UMeshComponent*> SelectedComponents = MeshToolManager->GetSelectedMeshComponents();
+		TArray<UMeshComponent*> SelectedComponents = SharedMeshToolData->GetSelectedMeshComponents();
 
 		for (UMeshComponent* MeshComponent : SelectedComponents)
 		{
@@ -635,9 +638,9 @@ void UMeshColorPaintingTool::LODPaintStateChanged(const bool bLODPaintingEnabled
 
 	ApplyForcedLODIndex(bLODPaintingEnabled ? CachedLODIndex : -1);
 
-	if (UMeshToolManager* MeshToolManager = Cast<UMeshToolManager>(GetToolManager()))
+	if (SharedMeshToolData.IsValid())
 	{
-		TArray<UMeshComponent*> PaintableComponents = MeshToolManager->GetPaintableMeshComponents();
+		TArray<UMeshComponent*> PaintableComponents = SharedMeshToolData->GetPaintableMeshComponents();
 
 		TUniquePtr< FComponentReregisterContext > ComponentReregisterContext;
 		//Make sure all static mesh render is dirty since we change the force LOD
@@ -649,15 +652,15 @@ void UMeshColorPaintingTool::LODPaintStateChanged(const bool bLODPaintingEnabled
 			}
 		}
 
-		MeshToolManager->Refresh();
+		SharedMeshToolData->Refresh();
 	}
 }
 
 void UMeshColorPaintingTool::ApplyForcedLODIndex(int32 ForcedLODIndex)
 {
-	if (UMeshToolManager* MeshToolManager = Cast<UMeshToolManager>(GetToolManager()))
+	if (SharedMeshToolData.IsValid())
 	{
-		TArray<UMeshComponent*> PaintableComponents = MeshToolManager->GetPaintableMeshComponents();
+		TArray<UMeshComponent*> PaintableComponents = SharedMeshToolData->GetPaintableMeshComponents();
 
 		for (UMeshComponent* SelectedComponent : PaintableComponents)
 		{
@@ -680,9 +683,9 @@ void UMeshColorPaintingTool::PaintLODChanged()
 
 		TUniquePtr< FComponentReregisterContext > ComponentReregisterContext;
 		//Make sure all static mesh render is dirty since we change the force LOD
-		if (UMeshToolManager* MeshToolManager = Cast<UMeshToolManager>(GetToolManager()))
+		if (SharedMeshToolData.IsValid())
 		{
-			TArray<UMeshComponent*> PaintableComponents = MeshToolManager->GetPaintableMeshComponents();
+			TArray<UMeshComponent*> PaintableComponents = SharedMeshToolData->GetPaintableMeshComponents();
 
 			for (UMeshComponent* SelectedComponent : PaintableComponents)
 			{
@@ -692,7 +695,7 @@ void UMeshColorPaintingTool::PaintLODChanged()
 				}
 			}
 
-			MeshToolManager->Refresh();
+			SharedMeshToolData->Refresh();
 		}
 	}
 }
@@ -736,14 +739,14 @@ void UMeshWeightPaintingTool::Setup()
 void UMeshWeightPaintingTool::CacheSelectionData()
 {
 	Super::CacheSelectionData();
-	if (UMeshToolManager* MeshToolManager = Cast<UMeshToolManager>(GetToolManager()))
+	if (SharedMeshToolData.IsValid())
 	{
-		MeshToolManager->ClearPaintableMeshComponents();
+		SharedMeshToolData->ClearPaintableMeshComponents();
 		//Determine LOD level to use for painting(can only paint on LODs in vertex mode)
 		const int32 PaintLODIndex = 0;
 		//Determine UV channel to use while painting textures
 		const int32 UVChannel = 0;
-		MeshToolManager->CacheSelectionData(PaintLODIndex, UVChannel);
+		SharedMeshToolData->CacheSelectionData(PaintLODIndex, UVChannel);
 	}
 }
 
diff --git a/Engine/Plugins/MeshPainting/Source/MeshPaintingToolset/Public/MeshPaintHelpers.h b/Engine/Plugins/MeshPainting/Source/MeshPaintingToolset/Public/MeshPaintHelpers.h
--- a/Engine/Plugins/MeshPainting/Source/MeshPaintingToolset/Public/MeshPaintHelpers.h
+++ b/Engine/Plugins/MeshPainting/Source/MeshPaintingToolset/Public/MeshPaintHelpers.h
@@ -251,22 +251,12 @@ void UMeshPaintingToolset::ApplyBrushToVertex(const FVector& VertexPosition, con
 };
 
 UCLASS()
-class MESHPAINTINGTOOLSET_API UMeshToolsContext : public UEdModeInteractiveToolsContext
-{
-	GENERATED_BODY()
-
-public:
-	UMeshToolsContext();
-};
-
-UCLASS()
-class MESHPAINTINGTOOLSET_API UMeshToolManager : public UInteractiveToolManager
+class MESHPAINTINGTOOLSET_API UMeshToolManager : public UObject
 {
 	GENERATED_BODY()
 
 public:
 	UMeshToolManager();
-	virtual void Shutdown() override;
 
 	/** Map of geometry adapters for each selected mesh component */
 	TMap<UMeshComponent*, TSharedPtr<IMeshPaintComponentAdapter>> GetComponentToAdapterMap() const;
diff --git a/Engine/Plugins/MeshPainting/Source/MeshPaintingToolset/Public/MeshPaintInteractions.h b/Engine/Plugins/MeshPainting/Source/MeshPaintingToolset/Public/MeshPaintInteractions.h
--- a/Engine/Plugins/MeshPainting/Source/MeshPaintingToolset/Public/MeshPaintInteractions.h
+++ b/Engine/Plugins/MeshPainting/Source/MeshPaintingToolset/Public/MeshPaintInteractions.h
@@ -10,6 +10,7 @@
 class IMeshPaintComponentAdapter;
 class AActor;
 class UMeshComponent;
+class UMeshToolManager;
 
 UINTERFACE()
 class MESHPAINTINGTOOLSET_API UMeshPaintSelectionInterface : public UInterface
@@ -38,13 +39,16 @@ public:
 	{
 		bAddToSelectionSet = bInNewSelectionType;
 	}
+	void SetMeshToolData(TWeakObjectPtr<UMeshToolManager> InMeshToolData);
 protected:
-	bool FindClickedComponentsAndCacheAdapters(const FInputDeviceRay& ClickPos, class UMeshToolManager* MeshToolManager);
+	bool FindClickedComponentsAndCacheAdapters(const FInputDeviceRay& ClickPos);
 
 protected:
 	UPROPERTY(Transient)
 	TArray<UMeshComponent*> CachedClickedComponents;
 	UPROPERTY(Transient)
 	TArray<AActor*> CachedClickedActors;
+
+	TWeakObjectPtr<UMeshToolManager> SharedMeshToolData;
 	bool bAddToSelectionSet;
 };
diff --git a/Engine/Plugins/MeshPainting/Source/MeshPaintingToolset/Public/MeshSelect.h b/Engine/Plugins/MeshPainting/Source/MeshPaintingToolset/Public/MeshSelect.h
--- a/Engine/Plugins/MeshPainting/Source/MeshPaintingToolset/Public/MeshSelect.h
+++ b/Engine/Plugins/MeshPainting/Source/MeshPaintingToolset/Public/MeshSelect.h
@@ -11,6 +11,7 @@ class IMeshPaintComponentAdapter;
 class AActor;
 class UMeshComponent;
 class UMeshPaintSelectionMechanic;
+class UMeshToolManager;
 
 /**
  * Builder for UVertexAdapterClickTool
@@ -23,6 +24,8 @@ class MESHPAINTINGTOOLSET_API UVertexAdapterClickToolBuilder : public USingleCli
 public:
 	virtual bool CanBuildTool(const FToolBuilderState& SceneState) const override;
 	virtual UInteractiveTool* BuildTool(const FToolBuilderState& SceneState) const override;
+	
+	TWeakObjectPtr<UMeshToolManager> SharedMeshToolData;
 };
 
 
@@ -37,6 +40,8 @@ class MESHPAINTINGTOOLSET_API UTextureAdapterClickToolBuilder : public USingleCl
 public:
 	virtual bool CanBuildTool(const FToolBuilderState& SceneState) const override;
 	virtual UInteractiveTool* BuildTool(const FToolBuilderState& SceneState) const override;
+
+	TWeakObjectPtr<UMeshToolManager> SharedMeshToolData;
 };
 
 
@@ -75,6 +80,8 @@ public:
 		return true;
 	}
 
+	void SetMeshToolData(TWeakObjectPtr<UMeshToolManager> InMeshToolData);
+
 protected:
 	// flags used to identify modifier keys/buttons
 	static const int AdditiveSelectionModifier = 1;
@@ -82,6 +89,8 @@ protected:
 
 	UPROPERTY(Transient)
 	UMeshPaintSelectionMechanic* SelectionMechanic;
+
+	TWeakObjectPtr<UMeshToolManager> SharedMeshToolData;
 };
 
 UCLASS()
diff --git a/Engine/Plugins/MeshPainting/Source/MeshPaintingToolset/Public/MeshTexturePaintingTool.h b/Engine/Plugins/MeshPainting/Source/MeshPaintingToolset/Public/MeshTexturePaintingTool.h
--- a/Engine/Plugins/MeshPainting/Source/MeshPaintingToolset/Public/MeshTexturePaintingTool.h
+++ b/Engine/Plugins/MeshPainting/Source/MeshPaintingToolset/Public/MeshTexturePaintingTool.h
@@ -35,6 +35,8 @@ class MESHPAINTINGTOOLSET_API UMeshTexturePaintingToolBuilder : public UInteract
 public:
 	virtual bool CanBuildTool(const FToolBuilderState& SceneState) const override;
 	virtual UInteractiveTool* BuildTool(const FToolBuilderState& SceneState) const override;
+
+	TWeakObjectPtr<UMeshToolManager> SharedMeshToolData;
 };
 
 
@@ -133,6 +135,9 @@ public:
 
 	bool ShouldFilterTextureAsset(const FAssetData& AssetData) const;
 	void PaintTextureChanged(const FAssetData& AssetData);
+
+	void SetMeshToolData(TWeakObjectPtr<UMeshToolManager> InToolData);
+
 protected:
 	virtual void SetAdditionalPaintParameters(FMeshPaintParameters& InPaintParameters) {};
 	virtual void FinishPainting();
@@ -166,6 +171,8 @@ private:
 	UPROPERTY(Transient)
 	TArray<const UTexture*> Textures;
 
+	TWeakObjectPtr<UMeshToolManager> SharedMeshToolData;
+
 	/** Flag for whether or not we are currently painting */
 	bool bArePainting;
 	bool bDoRestoreRenTargets;
diff --git a/Engine/Plugins/MeshPainting/Source/MeshPaintingToolset/Public/MeshVertexPaintingTool.h b/Engine/Plugins/MeshPainting/Source/MeshPaintingToolset/Public/MeshVertexPaintingTool.h
--- a/Engine/Plugins/MeshPainting/Source/MeshPaintingToolset/Public/MeshVertexPaintingTool.h
+++ b/Engine/Plugins/MeshPainting/Source/MeshPaintingToolset/Public/MeshVertexPaintingTool.h
@@ -63,6 +63,8 @@ class MESHPAINTINGTOOLSET_API UMeshColorPaintingToolBuilder : public UInteractiv
 public:
 	virtual bool CanBuildTool(const FToolBuilderState& SceneState) const override;
 	virtual UInteractiveTool* BuildTool(const FToolBuilderState& SceneState) const override;
+	
+	TWeakObjectPtr<UMeshToolManager> SharedMeshToolData;
 };
 
 UCLASS()
@@ -73,6 +75,8 @@ class MESHPAINTINGTOOLSET_API UMeshWeightPaintingToolBuilder : public UInteracti
 public:
 	virtual bool CanBuildTool(const FToolBuilderState& SceneState) const override;
 	virtual UInteractiveTool* BuildTool(const FToolBuilderState& SceneState) const override;
+
+	TWeakObjectPtr<UMeshToolManager> SharedMeshToolData;
 };
 
 
@@ -200,6 +204,8 @@ public:
 		return true;
 	}
 
+	void SetMeshToolData(TWeakObjectPtr<UMeshToolManager> InMeshToolData);
+
 protected:
 	virtual void SetAdditionalPaintParameters(FMeshPaintParameters& InPaintParameters) {};
 	virtual void FinishPainting();
@@ -225,12 +231,15 @@ protected:
 	UPROPERTY(Transient)
 	UMeshPaintSelectionMechanic* SelectionMechanic;
 
+	TWeakObjectPtr<UMeshToolManager> SharedMeshToolData;
+
 private:
 	bool PaintInternal(const TArrayView<TPair<FVector, FVector>>& Rays, EMeshPaintModeAction PaintAction, float PaintStrength);
 	
 private:
 	UPROPERTY(Transient)
 	UMeshVertexPaintingToolProperties* VertexProperties;
+
 	/** Flag for whether or not we are currently painting */
 	bool bArePainting;
 	bool bDoRestoreRenTargets;
@@ -268,6 +277,7 @@ protected:
 private:
 	UPROPERTY(Transient)
 	UMeshColorPaintingToolProperties* ColorProperties;
+
 	/** Current LOD index used for painting / forcing */
 	int32 CachedLODIndex;
 	/** Whether or not a specific LOD level should be forced */
diff --git a/Engine/Source/Editor/EditorFramework/Public/Toolkits/IToolkit.h b/Engine/Source/Editor/EditorFramework/Public/Toolkits/IToolkit.h
--- a/Engine/Source/Editor/EditorFramework/Public/Toolkits/IToolkit.h
+++ b/Engine/Source/Editor/EditorFramework/Public/Toolkits/IToolkit.h
@@ -127,7 +127,7 @@ public:
 	virtual FEdMode* GetEditorMode() const = 0;
 
 	/** @returns the scriptable editor mode this toolkit is used for, or null if not relevant. */
-	virtual UEdMode* GetScriptableEditorMode() const = 0;
+	virtual TWeakObjectPtr<UEdMode> GetScriptableEditorMode() const = 0;
 
 	/** @return the editor mode display name if there is an editor mode associated with this toolkit.  An empty FText is returned otherwise */
 	virtual FText GetEditorModeDisplayName() const = 0;
diff --git a/Engine/Source/Editor/Experimental/EditorInteractiveToolsFramework/Private/EdModeInteractiveToolsContext.cpp b/Engine/Source/Editor/Experimental/EditorInteractiveToolsFramework/Private/EdModeInteractiveToolsContext.cpp
--- a/Engine/Source/Editor/Experimental/EditorInteractiveToolsFramework/Private/EdModeInteractiveToolsContext.cpp
+++ b/Engine/Source/Editor/Experimental/EditorInteractiveToolsFramework/Private/EdModeInteractiveToolsContext.cpp
@@ -1075,7 +1075,7 @@ FRay UEdModeInteractiveToolsContext::GetRayFromMousePos(FEditorViewportClient* V
 		ViewportClient->EngineShowFlags).SetRealtimeUpdate(ViewportClient->IsRealtime()));		// why SetRealtimeUpdate here??
 	// this View is deleted by the FSceneViewFamilyContext destructor
 	FSceneView* View = ViewportClient->CalcSceneView(&ViewFamily);
-	FViewportCursorLocation MouseViewportRay(View, (FEditorViewportClient*)Viewport->GetClient(), MouseX, MouseY);
+	FViewportCursorLocation MouseViewportRay(View, ViewportClient, MouseX, MouseY);
 
 	FVector RayOrigin = MouseViewportRay.GetOrigin();
 	FVector RayDirection = MouseViewportRay.GetDirection();
diff --git a/Engine/Source/Editor/Experimental/EditorInteractiveToolsFramework/Public/EdModeInteractiveToolsContext.h b/Engine/Source/Editor/Experimental/EditorInteractiveToolsFramework/Public/EdModeInteractiveToolsContext.h
--- a/Engine/Source/Editor/Experimental/EditorInteractiveToolsFramework/Public/EdModeInteractiveToolsContext.h
+++ b/Engine/Source/Editor/Experimental/EditorInteractiveToolsFramework/Public/EdModeInteractiveToolsContext.h
@@ -26,14 +26,14 @@ class FViewportClient;
  * conversions and forwarding calls necessary to operate the ToolsContext
  */
 UCLASS(Transient)
-class EDITORINTERACTIVETOOLSFRAMEWORK_API UEdModeInteractiveToolsContext : public UInteractiveToolsContext
+class EDITORINTERACTIVETOOLSFRAMEWORK_API UEdModeInteractiveToolsContext final : public UInteractiveToolsContext
 {
 	GENERATED_BODY()
 
 public:
 	UEdModeInteractiveToolsContext();
 
-
+	UE_DEPRECATED(5.0, "You should no longer create your own EdModeInteractiveToolsContext; use the one in the FEditorModeTools::GetInteractiveToolsContext insetead.")
 	virtual void InitializeContextFromEdMode(FEdMode* EditorModeIn,
 		IToolsContextAssetAPI* UseAssetAPI = nullptr);
 	virtual void InitializeContextWithEditorModeManager(FEditorModeTools* InEditorModeManager,
diff --git a/Engine/Source/Editor/LevelEditor/Private/SLevelEditorToolBox.cpp b/Engine/Source/Editor/LevelEditor/Private/SLevelEditorToolBox.cpp
--- a/Engine/Source/Editor/LevelEditor/Private/SLevelEditorToolBox.cpp
+++ b/Engine/Source/Editor/LevelEditor/Private/SLevelEditorToolBox.cpp
@@ -173,12 +173,12 @@ void SLevelEditorToolBox::UpdateInlineContent(const TSharedPtr<IToolkit>& Toolki
 
 	if (Toolkit.IsValid())
 	{
-		if(Toolkit->GetEditorMode() || Toolkit->GetScriptableEditorMode())
+		if(Toolkit->GetEditorMode() || Toolkit->GetScriptableEditorMode().IsValid())
 		{
 			TabName = Toolkit->GetEditorModeDisplayName();
 			TabIcon = Toolkit->GetEditorModeIcon().GetSmallIcon();
 
-			TSharedPtr<FModeToolkit> ModeToolkit = StaticCastSharedPtr<FModeToolkit>(Toolkit);
+			TWeakPtr<FModeToolkit> ModeToolkit = StaticCastSharedPtr<FModeToolkit>(Toolkit);
 
 			TSharedRef< SUniformWrapPanel> PaletteTabBox = SNew(SUniformWrapPanel)
 				.SlotPadding( FMargin(1.f, 2.f) )
@@ -189,15 +189,16 @@ void SLevelEditorToolBox::UpdateInlineContent(const TSharedPtr<IToolkit>& Toolki
 
 			// Also build the toolkit here 
 			TArray<FName> PaletteNames;
-			ModeToolkit->GetToolPaletteNames(PaletteNames);
+			ModeToolkit.Pin()->GetToolPaletteNames(PaletteNames);
 
 			TSharedPtr<FUICommandList> CommandList;
-			CommandList = ModeToolkit->GetToolkitCommands();
+			CommandList = ModeToolkit.Pin()->GetToolkitCommands();
 
 			TSharedRef< SWidgetSwitcher > PaletteSwitcher = SNew(SWidgetSwitcher)
 			.WidgetIndex_Lambda( [PaletteNames, ModeToolkit] () -> int32 { 
 				int32 FoundIndex;
-				if (PaletteNames.Find(ModeToolkit->GetCurrentPalette(), FoundIndex))
+				TSharedPtr<FModeToolkit> PinnedToolkit = ModeToolkit.Pin();
+				if (PinnedToolkit.IsValid() && PaletteNames.Find(PinnedToolkit->GetCurrentPalette(), FoundIndex))
 				{
 					return FoundIndex;	
 				}
@@ -206,11 +207,11 @@ void SLevelEditorToolBox::UpdateInlineContent(const TSharedPtr<IToolkit>& Toolki
 			
 			for(auto Palette : PaletteNames)
 			{
-				FName ToolbarCustomizationName = ModeToolkit->GetEditorMode() ? ModeToolkit->GetEditorMode()->GetModeInfo().ToolbarCustomizationName : ModeToolkit->GetScriptableEditorMode()->GetModeInfo().ToolbarCustomizationName;
+				FName ToolbarCustomizationName = ModeToolkit.Pin()->GetEditorMode() ? ModeToolkit.Pin()->GetEditorMode()->GetModeInfo().ToolbarCustomizationName : ModeToolkit.Pin()->GetScriptableEditorMode()->GetModeInfo().ToolbarCustomizationName;
 				FUniformToolBarBuilder ModeToolbarBuilder(CommandList, FMultiBoxCustomization(ToolbarCustomizationName));
 				ModeToolbarBuilder.SetStyle(&FEditorStyle::Get(), "PaletteToolBar");
 
-				ModeToolkit->BuildToolPalette(Palette, ModeToolbarBuilder);
+				ModeToolkit.Pin()->BuildToolPalette(Palette, ModeToolbarBuilder);
 
 				TSharedRef<SWidget> PaletteWidget = ModeToolbarBuilder.MakeWidget();
 
@@ -219,16 +220,27 @@ void SLevelEditorToolBox::UpdateInlineContent(const TSharedPtr<IToolkit>& Toolki
 					SNew(SCheckBox)
 					.Padding(FMargin(8.f, 4.f, 8.f, 5.f))
 					.Style( FEditorStyle::Get(),  "PaletteToolBar.Tab" )
-					.OnCheckStateChanged_Lambda( [/*PaletteSwitcher, PaletteWidget, */ModeToolkit, Palette] (const ECheckBoxState) { 
-							ModeToolkit->SetCurrentPalette(Palette);
+					.OnCheckStateChanged_Lambda([/*PaletteSwitcher, PaletteWidget, */ModeToolkit, Palette](const ECheckBoxState) {
+							TSharedPtr<FModeToolkit> PinnedToolkit = ModeToolkit.Pin();
+							if (PinnedToolkit.IsValid())
+							{
+								PinnedToolkit->SetCurrentPalette(Palette);
+							}
 						} 
 					)
 					// .IsChecked_Lambda( [PaletteSwitcher, PaletteWidget] () -> ECheckBoxState { return PaletteSwitcher->GetActiveWidget() == PaletteWidget ? ECheckBoxState::Checked : ECheckBoxState::Unchecked; })
-					.IsChecked_Lambda( [ModeToolkit, Palette] () -> ECheckBoxState { return ModeToolkit->GetCurrentPalette() == Palette ? ECheckBoxState::Checked : ECheckBoxState::Unchecked; } )
+					.IsChecked_Lambda([ModeToolkit, Palette]() -> ECheckBoxState {
+						TSharedPtr<FModeToolkit> PinnedToolkit = ModeToolkit.Pin();
+						if (PinnedToolkit.IsValid() && (PinnedToolkit->GetCurrentPalette() == Palette))
+						{
+							return ECheckBoxState::Checked;
+						}
+						return ECheckBoxState::Unchecked;
+					})
 					[
 						SNew(STextBlock)
 						.TextStyle(FAppStyle::Get(), "NormalText")
-						.Text(ModeToolkit->GetToolPaletteDisplayName(Palette))
+						.Text(ModeToolkit.Pin()->GetToolPaletteDisplayName(Palette))
 						.Justification(ETextJustify::Center)
 					]
 				];
@@ -264,7 +276,7 @@ void SLevelEditorToolBox::UpdateInlineContent(const TSharedPtr<IToolkit>& Toolki
 			}
 
 			StatusBarMessageHandle = GEditor->GetEditorSubsystem<UStatusBarSubsystem>()->PushStatusBarMessage(LevelEditorStatusBarName, TAttribute<FText>::Create(
-				TAttribute<FText>::FGetter::CreateSP(ModeToolkit.Get(), &FModeToolkit::GetActiveToolDisplayName))
+				TAttribute<FText>::FGetter::CreateSP(ModeToolkit.Pin().ToSharedRef(), &FModeToolkit::GetActiveToolDisplayName))
 			);
 		}
 	}
diff --git a/Engine/Source/Editor/UnrealEd/Private/EditorModeManager.cpp b/Engine/Source/Editor/UnrealEd/Private/EditorModeManager.cpp
--- a/Engine/Source/Editor/UnrealEd/Private/EditorModeManager.cpp
+++ b/Engine/Source/Editor/UnrealEd/Private/EditorModeManager.cpp
@@ -32,6 +32,7 @@
 #include "Widgets/Images/SImage.h"
 #include "InputRouter.h"
 #include "InteractiveGizmoManager.h"
+#include "EdModeInteractiveToolsContext.h"
 
 /*------------------------------------------------------------------------------
 	FEditorModeTools.
@@ -62,6 +63,9 @@ FEditorModeTools::FEditorModeTools()
 {
 	DefaultModeIDs.Add( FBuiltinEditorModes::EM_Default );
 
+	InteractiveToolsContext = NewObject<UEdModeInteractiveToolsContext>(GetTransientPackage(), UEdModeInteractiveToolsContext::StaticClass(), TEXT("ToolsContext"), RF_Transient);
+	InteractiveToolsContext->InitializeContextWithEditorModeManager(this);
+
 	// Load the last used settings
 	LoadConfig();
 
@@ -86,6 +90,9 @@ FEditorModeTools::~FEditorModeTools()
 
 	DeactivateAllModes();
 	RecycledScriptableModes.Empty();
+
+	InteractiveToolsContext->ShutdownContext();
+	InteractiveToolsContext = nullptr;
 }
 
 void FEditorModeTools::LoadConfig(void)
@@ -1557,6 +1564,7 @@ void FEditorModeTools::AddReferencedObjects( FReferenceCollector& Collector )
 {
 	Collector.AddReferencedObjects(ActiveScriptableModes);
 	Collector.AddReferencedObjects(RecycledScriptableModes);
+	Collector.AddReferencedObject(InteractiveToolsContext);
 }
 
 FEdMode* FEditorModeTools::GetActiveMode( FEditorModeID InID )
@@ -1645,3 +1653,8 @@ bool FEditorModeTools::CanAutoSave() const
 
 	return true;
 }
+
+UEdModeInteractiveToolsContext* FEditorModeTools::GetInteractiveToolsContext() const
+{
+	return InteractiveToolsContext;
+}
diff --git a/Engine/Source/Editor/UnrealEd/Private/Toolkits/AssetEditorToolkit.cpp b/Engine/Source/Editor/UnrealEd/Private/Toolkits/AssetEditorToolkit.cpp
--- a/Engine/Source/Editor/UnrealEd/Private/Toolkits/AssetEditorToolkit.cpp
+++ b/Engine/Source/Editor/UnrealEd/Private/Toolkits/AssetEditorToolkit.cpp
@@ -405,7 +405,7 @@ class FEdMode* FAssetEditorToolkit::GetEditorMode() const
 	return nullptr;
 }
 
-class UEdMode* FAssetEditorToolkit::GetScriptableEditorMode() const
+TWeakObjectPtr<UEdMode> FAssetEditorToolkit::GetScriptableEditorMode() const
 {
 	return nullptr;
 }
diff --git a/Engine/Source/Editor/UnrealEd/Private/Toolkits/BaseToolkit.cpp b/Engine/Source/Editor/UnrealEd/Private/Toolkits/BaseToolkit.cpp
--- a/Engine/Source/Editor/UnrealEd/Private/Toolkits/BaseToolkit.cpp
+++ b/Engine/Source/Editor/UnrealEd/Private/Toolkits/BaseToolkit.cpp
@@ -123,16 +123,22 @@ FEditorModeTools& FBaseToolkit::GetEditorModeManager() const
 
 
 void FModeToolkit::Init(const TSharedPtr< class IToolkitHost >& InitToolkitHost)
+{
+	Init(InitToolkitHost, TWeakObjectPtr<UEdMode>());
+}
+
+void FModeToolkit::Init(const TSharedPtr< class IToolkitHost >& InitToolkitHost, TWeakObjectPtr<UEdMode> InOwningMode)
 {
 	check( InitToolkitHost.IsValid() );
 	
 	ToolkitMode = EToolkitMode::WorldCentric;
 	ToolkitHost = InitToolkitHost;
+	OwningEditorMode = InOwningMode;
 
-	if (UEdMode* ScriptableEditorMode = GetScriptableEditorMode())
+	if (OwningEditorMode.IsValid())
 	{
-		ScriptableEditorMode->GetToolManager()->OnToolStarted.AddSP(this, &FModeToolkit::OnToolStarted);
-		ScriptableEditorMode->GetToolManager()->OnToolEnded.AddSP(this, &FModeToolkit::OnToolEnded);
+		OwningEditorMode->GetToolManager()->OnToolStarted.AddSP(this, &FModeToolkit::OnToolStarted);
+		OwningEditorMode->GetToolManager()->OnToolEnded.AddSP(this, &FModeToolkit::OnToolEnded);
 	}
 
 	FPropertyEditorModule& PropertyEditorModule = FModuleManager::GetModuleChecked<FPropertyEditorModule>("PropertyEditor");
@@ -174,14 +180,13 @@ void FModeToolkit::Init(const TSharedPtr< class IToolkitHost >& InitToolkitHost)
 
 FModeToolkit::~FModeToolkit()
 {
-	if (UEdMode* ScriptableEditorMode = GetScriptableEditorMode())
+	if (OwningEditorMode.IsValid())
 	{
-		if (ScriptableEditorMode->GetToolManager())
-		{
-			ScriptableEditorMode->GetToolManager()->OnToolStarted.RemoveAll(this);
-			ScriptableEditorMode->GetToolManager()->OnToolEnded.RemoveAll(this);
-		}
+		OwningEditorMode->GetToolManager()->OnToolStarted.RemoveAll(this);
+		OwningEditorMode->GetToolManager()->OnToolEnded.RemoveAll(this);
 	}
+
+	OwningEditorMode.Reset();
 }
 
 FName FModeToolkit::GetToolkitFName() const
@@ -218,7 +223,12 @@ FLinearColor FModeToolkit::GetWorldCentricTabColorScale() const
 
 bool FModeToolkit::CanStartTool(const FString& ToolTypeIdentifier)
 {
-	UInteractiveToolManager* Manager = GetScriptableEditorMode()->GetToolManager();
+	if (!OwningEditorMode.IsValid())
+	{
+		return false;
+	}
+
+	UInteractiveToolManager* Manager = OwningEditorMode->GetToolManager();
 
 	return (Manager->HasActiveTool(EToolSide::Left) == false) &&
 		(Manager->CanActivateTool(EToolSide::Left, ToolTypeIdentifier) == true);
@@ -226,24 +236,44 @@ bool FModeToolkit::CanStartTool(const FString& ToolTypeIdentifier)
 
 bool FModeToolkit::CanAcceptActiveTool()
 {
-	return GetScriptableEditorMode()->GetToolManager()->CanAcceptActiveTool(EToolSide::Left);
+	if (!OwningEditorMode.IsValid())
+	{
+		return false;
+	}
+
+	return OwningEditorMode->GetToolManager()->CanAcceptActiveTool(EToolSide::Left);
 }
 
 bool FModeToolkit::CanCancelActiveTool()
 {
-	return GetScriptableEditorMode()->GetToolManager()->CanCancelActiveTool(EToolSide::Left);
+	if (!OwningEditorMode.IsValid())
+	{
+		return false;
+	}
+
+	return OwningEditorMode->GetToolManager()->CanCancelActiveTool(EToolSide::Left);
 }
 
 bool FModeToolkit::CanCompleteActiveTool()
 {
-	return GetScriptableEditorMode()->GetToolManager()->HasActiveTool(EToolSide::Left) && CanCancelActiveTool() == false;
+	if (!OwningEditorMode.IsValid())
+	{
+		return false;
+	}
+
+	return OwningEditorMode->GetToolManager()->HasActiveTool(EToolSide::Left) && CanCancelActiveTool() == false;
 }
 
 
 void FModeToolkit::OnToolStarted(UInteractiveToolManager* Manager, UInteractiveTool* Tool)
 {
 	// Update properties panel
-	UInteractiveTool* CurTool = GetScriptableEditorMode()->GetToolManager()->GetActiveTool(EToolSide::Left);
+	if (!OwningEditorMode.IsValid())
+	{
+		return;
+	}
+
+	UInteractiveTool* CurTool = OwningEditorMode->GetToolManager()->GetActiveTool(EToolSide::Left);
 	if (CurTool)
 	{
 		DetailsView->SetObjects(CurTool->GetToolProperties());
@@ -266,9 +296,9 @@ FText FModeToolkit::GetEditorModeDisplayName() const
 	{
 		return EdMode->GetModeInfo().Name;
 	}
-	else if (UEdMode* ScriptableMode = GetScriptableEditorMode())
+	else if (OwningEditorMode.IsValid())
 	{
-		return ScriptableMode->GetModeInfo().Name;
+		return OwningEditorMode->GetModeInfo().Name;
 	}
 
 	return FText::GetEmpty();
@@ -280,17 +310,17 @@ FSlateIcon FModeToolkit::GetEditorModeIcon() const
 	{
 		return EdMode->GetModeInfo().IconBrush;
 	}
-	else if (UEdMode* ScriptableMode = GetScriptableEditorMode())
+	else if (OwningEditorMode.IsValid())
 	{
-		return ScriptableMode->GetModeInfo().IconBrush;
+		return OwningEditorMode->GetModeInfo().IconBrush;
 	}
 
 	return FSlateIcon();
 }
 
-UEdMode* FModeToolkit::GetScriptableEditorMode() const
+TWeakObjectPtr<UEdMode> FModeToolkit::GetScriptableEditorMode() const
 {
-	return nullptr;
+	return OwningEditorMode;
 }
 
 TSharedPtr<SWidget> FModeToolkit::GetInlineContent() const
@@ -309,18 +339,19 @@ TSharedPtr<SWidget> FModeToolkit::GetInlineContent() const
 
 void FModeToolkit::BuildToolPalette(FName PaletteName, class FToolBarBuilder& ToolbarBuilder)
 {
-	if (UEdMode* EditorMode = GetScriptableEditorMode())
+	if (!OwningEditorMode.IsValid())
+	{
+		return;
+	}
+
+	TMap<FName, TArray<TSharedPtr<FUICommandInfo>>> CommandLists = OwningEditorMode->GetModeCommands();
+	TArray<TSharedPtr<FUICommandInfo>>* CurrentCommandListPtr = CommandLists.Find(PaletteName);
+	if (CurrentCommandListPtr)
 	{
-		TMap<FName, TArray<TSharedPtr<FUICommandInfo>>> CommandLists = EditorMode->GetModeCommands();
-		TArray<TSharedPtr<FUICommandInfo>>* CurrentCommandListPtr = CommandLists.Find(PaletteName);
-		if (CurrentCommandListPtr)
+		TArray<TSharedPtr<FUICommandInfo>> CurrentCommandList = *CurrentCommandListPtr;
+		for (TSharedPtr<FUICommandInfo> Command : CurrentCommandList)
 		{
-			TArray<TSharedPtr<FUICommandInfo>> CurrentCommandList = *CurrentCommandListPtr;
-			for (TSharedPtr<FUICommandInfo> Command : CurrentCommandList)
-			{
-				ToolbarBuilder.AddToolBarButton(Command);
-			}
-			
+			ToolbarBuilder.AddToolBarButton(Command);
 		}
 
 	}
diff --git a/Engine/Source/Editor/UnrealEd/Private/Tools/UEdMode.cpp b/Engine/Source/Editor/UnrealEd/Private/Tools/UEdMode.cpp
--- a/Engine/Source/Editor/UnrealEd/Private/Tools/UEdMode.cpp
+++ b/Engine/Source/Editor/UnrealEd/Private/Tools/UEdMode.cpp
@@ -20,7 +20,7 @@
 #include "Toolkits/ToolkitManager.h"
 #include "Materials/Material.h"
 #include "Materials/MaterialInterface.h"
-#include "LevelEditorViewport.h"
+#include "EditorViewportClient.h"
 #include "InteractiveToolManager.h"
 #include "InteractiveToolObjects.h"
 #include "Editor/EditorEngine.h"
@@ -67,7 +67,6 @@ UEdMode::UEdMode()
 {
 	bDrawKillZ = true;
 	ToolsContext = nullptr;
-	ToolsContextClass = UEdModeInteractiveToolsContext::StaticClass();
 	ToolCommandList = MakeShareable(new FUICommandList);
 }
 
@@ -216,16 +215,30 @@ void UEdMode::Enter()
 
 	bPendingDeletion = false;
 
-	// initialize the adapter that attaches the ToolsContext to this FEdMode
-	ToolsContext = NewObject<UEdModeInteractiveToolsContext>(GetTransientPackage(), ToolsContextClass.Get(), TEXT("ToolsContext"), RF_Transient);
-	ToolsContext->InitializeContextWithEditorModeManager(GetModeManager());
-
+	ToolsContext = Owner->GetInteractiveToolsContext();
+	check(ToolsContext.IsValid());
 
 	GetToolManager()->OnToolStarted.AddUObject(this, &UEdMode::OnToolStarted);
 	GetToolManager()->OnToolEnded.AddUObject(this, &UEdMode::OnToolEnded);
 
 	// Now that the context is ready, make the toolkit
 	CreateToolkit();
+	check(Toolkit.IsValid());
+
+	Toolkit->Init(Owner->GetToolkitHost(), this);
+
+	BindCommands();
+
+	if (SettingsClass.IsValid())
+	{
+		UClass* LoadedSettingsObject = SettingsClass.LoadSynchronous();
+		SettingsObject = NewObject<UObject>(this, LoadedSettingsObject);
+	}
+
+	if (SettingsObject)
+	{
+		Toolkit->SetModeSettingsObject(SettingsObject);
+	}
 
 	FEditorDelegates::EditorModeIDEnter.Broadcast(GetID());
 
@@ -239,6 +252,8 @@ void UEdMode::RegisterTool(TSharedPtr<FUICommandInfo> UICommand, FString ToolIde
 		FExecuteAction::CreateLambda([this, ToolIdentifier]() { this->ToolsContext->StartTool(ToolIdentifier); }),
 		FCanExecuteAction::CreateLambda([this, ToolIdentifier]() { return this->ToolsContext->ToolManager->CanActivateTool(EToolSide::Mouse, ToolIdentifier); }),
 		FIsActionChecked::CreateLambda([this, Builder]() { return this->ToolsContext->IsToolBuilderActive(EToolSide::Mouse, Builder); }));
+
+	RegisteredTools.Emplace(UICommand, ToolIdentifier);
 }
 
 void UEdMode::Exit()
@@ -246,8 +261,13 @@ void UEdMode::Exit()
 	GetToolManager()->OnToolStarted.RemoveAll(this);
 	GetToolManager()->OnToolEnded.RemoveAll(this);
 
-	ToolsContext->ShutdownContext();
-	ToolsContext = nullptr;
+	const TSharedRef<FUICommandList>& CommandList = Toolkit->GetToolkitCommands();
+	for (auto& RegisteredTool : RegisteredTools)
+	{
+		CommandList->UnmapAction(RegisteredTool.Key);
+		ToolsContext->ToolManager->UnregisterToolType(RegisteredTool.Value);
+	}
+	RegisteredTools.SetNum(0);
 
 	if (Toolkit.IsValid())
 	{
@@ -255,6 +275,8 @@ void UEdMode::Exit()
 		Toolkit.Reset();
 	}
 
+	ToolsContext = nullptr;
+
 	FEditorDelegates::EditorModeIDExit.Broadcast(GetID());
 }
 
@@ -266,7 +288,7 @@ UTexture2D* UEdMode::GetVertexTexture()
 void UEdMode::Render(const FSceneView* View, FViewport* Viewport, FPrimitiveDrawInterface* PDI)
 {
 	// give ToolsContext a chance to render
-	if (ToolsContext != nullptr)
+	if (ToolsContext.IsValid())
 	{
 		ToolsContext->Render(View, Viewport, PDI);
 	}
@@ -440,25 +462,23 @@ AActor* UEdMode::GetFirstSelectedActorInstance() const
 
 UInteractiveToolManager* UEdMode::GetToolManager() const
 {
-	return ToolsContext->ToolManager;
-}
-
-void UEdMode::CreateToolkit()
-{
-	if (!Toolkit.IsValid())
+	if (ToolsContext.IsValid())
 	{
-		Toolkit = MakeShareable(new FModeToolkit);
-		Toolkit->Init(Owner->GetToolkitHost());
+		return ToolsContext->ToolManager;
+	}
 
+	return nullptr;
+}
 
-	}
+TWeakObjectPtr<UEdModeInteractiveToolsContext> UEdMode::GetInteractiveToolsContext() const
+{
+	return ToolsContext;
+}
 
-	UClass* LoadedSettingsObject = SettingsClass.LoadSynchronous();
-	SettingsObject = NewObject<UObject>(this, LoadedSettingsObject);
-	if (SettingsObject)
-	{
-		Toolkit->SetModeSettingsObject(SettingsObject);
-	}
+void UEdMode::CreateToolkit()
+{
+	check(!Toolkit.IsValid())
+	Toolkit = MakeShareable(new FModeToolkit);
 }
 
 bool UEdMode::IsSnapRotationEnabled()
diff --git a/Engine/Source/Editor/UnrealEd/Public/EditorModeManager.h b/Engine/Source/Editor/UnrealEd/Public/EditorModeManager.h
--- a/Engine/Source/Editor/UnrealEd/Public/EditorModeManager.h
+++ b/Engine/Source/Editor/UnrealEd/Public/EditorModeManager.h
@@ -25,6 +25,7 @@ struct FViewportClick;
 class UEdMode;
 class UInteractiveGizmoManager;
 class UInputRouter;
+class UEdModeInteractiveToolsContext;
 
 /**
  * A helper class to store the state of the various editor modes.
@@ -482,6 +483,8 @@ public:
 
 	void RemoveAllDelegateHandlers();
 
+	UEdModeInteractiveToolsContext* GetInteractiveToolsContext() const;
+
 protected:
 	/** 
 	 * Delegate handlers
@@ -515,6 +518,7 @@ protected:
 
 	/** if true the current selection has a scene component */
 	bool bSelectionHasSceneComponent;
+
 private:
 	struct FEdModeToolbarRow
 	{
@@ -557,6 +561,8 @@ private:
 	/** Flag set between calls to StartTracking() and EndTracking() */
 	bool bIsTracking;
 
+	UEdModeInteractiveToolsContext* InteractiveToolsContext;
+
 	FEditorViewportClient* HoveredViewportClient = nullptr;
 	FEditorViewportClient* FocusedViewportClient = nullptr;
 };
diff --git a/Engine/Source/Editor/UnrealEd/Public/Toolkits/AssetEditorToolkit.h b/Engine/Source/Editor/UnrealEd/Public/Toolkits/AssetEditorToolkit.h
--- a/Engine/Source/Editor/UnrealEd/Public/Toolkits/AssetEditorToolkit.h
+++ b/Engine/Source/Editor/UnrealEd/Public/Toolkits/AssetEditorToolkit.h
@@ -130,7 +130,7 @@ public:
 	virtual FText GetToolkitToolTipText() const override;
 	virtual FString GetWorldCentricTabPrefix() const override = 0;	// Must implement in derived class!
 	virtual class FEdMode* GetEditorMode() const override;
-	virtual class UEdMode* GetScriptableEditorMode() const override;
+	virtual TWeakObjectPtr<class UEdMode> GetScriptableEditorMode() const final;
 	virtual FText GetEditorModeDisplayName() const override;
 	virtual FSlateIcon GetEditorModeIcon() const override;
 
diff --git a/Engine/Source/Editor/UnrealEd/Public/Toolkits/BaseToolkit.h b/Engine/Source/Editor/UnrealEd/Public/Toolkits/BaseToolkit.h
--- a/Engine/Source/Editor/UnrealEd/Public/Toolkits/BaseToolkit.h
+++ b/Engine/Source/Editor/UnrealEd/Public/Toolkits/BaseToolkit.h
@@ -101,6 +101,7 @@ public:
 
 	/** Initializes the mode toolkit */
 	virtual void Init(const TSharedPtr<IToolkitHost>& InitToolkitHost);
+	virtual void Init(const TSharedPtr<IToolkitHost>& InitToolkitHost, TWeakObjectPtr<UEdMode> InOwningMode);
 	~FModeToolkit();
 
 public:
@@ -125,7 +126,7 @@ public:
 	virtual FText GetEditorModeDisplayName() const override;
 	virtual FSlateIcon GetEditorModeIcon() const override;
 
-	virtual UEdMode* GetScriptableEditorMode() const;
+	virtual TWeakObjectPtr<UEdMode> GetScriptableEditorMode() const final;
 	virtual TSharedPtr<SWidget> GetInlineContent() const override;
 
 	/** Returns the number of Mode specific tabs in the mode toolbar **/
@@ -168,6 +169,8 @@ protected:
 	TSharedPtr<IDetailsView> ModeDetailsView;
 	TSharedPtr<IDetailsView> DetailsView;
 
+	TWeakObjectPtr<UEdMode> OwningEditorMode;
+
 	FName CurrentPaletteName;
 	FOnPaletteChanged OnPaletteChangedDelegate;
 };
diff --git a/Engine/Source/Editor/UnrealEd/Public/Tools/EditorToolAssetAPI.h b/Engine/Source/Editor/UnrealEd/Public/Tools/EditorToolAssetAPI.h
--- a/Engine/Source/Editor/UnrealEd/Public/Tools/EditorToolAssetAPI.h
+++ b/Engine/Source/Editor/UnrealEd/Public/Tools/EditorToolAssetAPI.h
@@ -23,15 +23,4 @@ public:
 	virtual void InteractiveSaveGeneratedAsset(UObject* Asset, UPackage* AssetPackage) override;
 	virtual void AutoSaveGeneratedAsset(UObject* Asset, UPackage* AssetPackage) override;
 	virtual void NotifyGeneratedAssetModified(UObject* Asset, UPackage* AssetPackage) override;
-
-	virtual AActor* GenerateStaticMeshActor(
-		UWorld* TargetWorld,
-		FTransform Transform,
-		FString ObjectBaseName,
-		FGeneratedStaticMeshAssetConfig&& AssetConfig) override
-	{
-		check(false);		// currently not supported in default editor tool API. 
-							// ModelingModeAssetAPI contains an experimental implementation
-		return nullptr;
-	}
 };
\ No newline at end of file
diff --git a/Engine/Source/Editor/UnrealEd/Public/Tools/UEdMode.h b/Engine/Source/Editor/UnrealEd/Public/Tools/UEdMode.h
--- a/Engine/Source/Editor/UnrealEd/Public/Tools/UEdMode.h
+++ b/Engine/Source/Editor/UnrealEd/Public/Tools/UEdMode.h
@@ -275,7 +275,8 @@ public:
 	/**
 	 * @return active ToolManager
 	 */
-	virtual UInteractiveToolManager* GetToolManager() const;
+	UInteractiveToolManager* GetToolManager() const;
+	TWeakObjectPtr<UEdModeInteractiveToolsContext> GetInteractiveToolsContext() const;
 
 	virtual TMap<FName, TArray<TSharedPtr<FUICommandInfo>>> GetModeCommands() const
 	{
@@ -284,18 +285,20 @@ public:
 
 protected:
 	virtual void CreateToolkit();
-	virtual void OnToolStarted(UInteractiveToolManager* Manager, UInteractiveTool* Tool) {};
-	virtual void OnToolEnded(UInteractiveToolManager* Manager, UInteractiveTool* Tool) {};
-	virtual void ActivateDefaultTool() {};
+	virtual void OnToolStarted(UInteractiveToolManager* Manager, UInteractiveTool* Tool) {}
+	virtual void OnToolEnded(UInteractiveToolManager* Manager, UInteractiveTool* Tool) {}
+	virtual void ActivateDefaultTool() {}
+	virtual void BindCommands() {}
 
 protected:
 
-	UPROPERTY()
-	UEdModeInteractiveToolsContext* ToolsContext;
-	TSoftClassPtr<UEdModeInteractiveToolsContext> ToolsContextClass;
+	TWeakObjectPtr<UEdModeInteractiveToolsContext> ToolsContext;
+
 	/** Command list lives here so that the key bindings on the commands can be processed in the viewport. */
 	TSharedPtr<FUICommandList> ToolCommandList;
 
+	TArray<TPair<TSharedPtr<FUICommandInfo>, FString>> RegisteredTools;
+
 	UPROPERTY()
 	TSoftClassPtr<UObject> SettingsClass;
 
diff --git a/Engine/Source/Runtime/Experimental/InteractiveToolsFramework/Private/InteractiveToolManager.cpp b/Engine/Source/Runtime/Experimental/InteractiveToolsFramework/Private/InteractiveToolManager.cpp
--- a/Engine/Source/Runtime/Experimental/InteractiveToolsFramework/Private/InteractiveToolManager.cpp
+++ b/Engine/Source/Runtime/Experimental/InteractiveToolsFramework/Private/InteractiveToolManager.cpp
@@ -61,6 +61,19 @@ void UInteractiveToolManager::RegisterToolType(const FString& Identifier, UInter
 	ToolBuilders.Add(Identifier, Builder );
 }
 
+void UInteractiveToolManager::UnregisterToolType(const FString& Identifier)
+{
+	if (ActiveLeftToolName == Identifier)
+	{
+		DeactivateTool(EToolSide::Left, EToolShutdownType::Cancel);
+	}
+	if (ActiveRightToolName == Identifier)
+	{
+		DeactivateTool(EToolSide::Right, EToolShutdownType::Cancel);
+	}
+	ToolBuilders.Remove(Identifier);
+}
+
 
 bool UInteractiveToolManager::SelectActiveToolType(EToolSide Side, const FString& Identifier)
 {
diff --git a/Engine/Source/Runtime/Experimental/InteractiveToolsFramework/Public/InteractiveToolManager.h b/Engine/Source/Runtime/Experimental/InteractiveToolsFramework/Public/InteractiveToolManager.h
--- a/Engine/Source/Runtime/Experimental/InteractiveToolsFramework/Public/InteractiveToolManager.h
+++ b/Engine/Source/Runtime/Experimental/InteractiveToolsFramework/Public/InteractiveToolManager.h
@@ -93,6 +93,13 @@ public:
 	 */
 	virtual void RegisterToolType(const FString& Identifier, UInteractiveToolBuilder* Builder);
 
+	/**
+	 * Unregisters a ToolBuilder
+	 * @param Identifier string used to identify this Builder
+	 * @param Builder new ToolBuilder instance
+	 */
+	virtual void UnregisterToolType(const FString& Identifier);
+
 	/**
 	 * Set active ToolBuilder for a ToolSide via string identifier
 	 * @param Side which "side" should we set this Builder on
diff --git a/Engine/Source/Runtime/Experimental/InteractiveToolsFramework/Public/ToolContextInterfaces.h b/Engine/Source/Runtime/Experimental/InteractiveToolsFramework/Public/ToolContextInterfaces.h
--- a/Engine/Source/Runtime/Experimental/InteractiveToolsFramework/Public/ToolContextInterfaces.h
+++ b/Engine/Source/Runtime/Experimental/InteractiveToolsFramework/Public/ToolContextInterfaces.h
@@ -425,38 +425,5 @@ public:
 
 	/** Notify that asset has been created and is dirty */
 	virtual void NotifyGeneratedAssetModified(UObject* Asset, UPackage* AssetPackage) = 0;
-
-
-
-	/**
-	 * Create a new UStaticMeshAsset for the provided mesh, then a new UStaticMeshComponent and AStaticMeshActor in the TargetWorld
-	 * @param TargetWorld world to create new Actor in
-	 * @param Transform transform to set on Actor
-	 * @param ObjectBaseName a base name for the Asset and Actor. Unique name will be generated by appending a number, if a name collision occurs.
-	 * @param AssetConfig configuration information for the new Asset
-	 * @return new AStaticMeshActor with new Asset assigned to it's UStaticMeshComponent
-	 * @warning this function may return null. In this case the asset creation failed and/or the user cancelled during the process, if it was interactive.
-	 */
-	virtual AActor* GenerateStaticMeshActor(
-		UWorld* TargetWorld,
-		FTransform Transform,
-		FString ObjectBaseName,
-		FGeneratedStaticMeshAssetConfig&& AssetConfig) = 0;
-
-
-	/**
-	 * Save a generated UTexture2D as an Asset.
-	 * Assumption is that the texture was generated in code and is in the Transient package.
-	 * @param GeneratedTexture Texture2D to save
-	 * @param ObjectBaseName a base name for the Asset. Unique name will be generated by appending a number, if a name collision occurs.
-	 * @param RelativeToAsset New texture will be saved at the same path as the UPackage for this UObject, which we assume to be an Asset (eg like a UStaticMesh)
-	 */
-	virtual bool SaveGeneratedTexture2D(
-		UTexture2D* GeneratedTexture,
-		FString ObjectBaseName,
-		const UObject* RelativeToAsset)
-	{
-		check(false); return false;
-	}
 };
 
