MeshModelingTools: Made it possible to change tools without hitting accept. This involved making CanStartTool not check that no tool is active, so places where that function was used were tweaked to keep functionality the same. Also made buttons be highlighted when a specific tool is used.

#rnx

diff --git a/Engine/Plugins/Experimental/HairStrands/Source/HairStrandsEditor/Private/GroomEditorMode.cpp b/Engine/Plugins/Experimental/HairStrands/Source/HairStrandsEditor/Private/GroomEditorMode.cpp
--- a/Engine/Plugins/Experimental/HairStrands/Source/HairStrandsEditor/Private/GroomEditorMode.cpp
+++ b/Engine/Plugins/Experimental/HairStrands/Source/HairStrandsEditor/Private/GroomEditorMode.cpp
@@ -217,7 +217,7 @@ void FGroomEditorMode::Enter()
 		ToolsContext->ToolManager->RegisterToolType(ToolIdentifier, Builder);
 		CommandList->MapAction( UICommand,
 			FExecuteAction::CreateLambda([this, ToolIdentifier]() { ToolsContext->StartTool(ToolIdentifier); }),
-			FCanExecuteAction::CreateLambda([this, ToolIdentifier]() { return ToolsContext->CanStartTool(ToolIdentifier); }));
+			FCanExecuteAction::CreateLambda([this, ToolIdentifier]() { return !ToolsContext->HasActiveTool() && ToolsContext->CanStartTool(ToolIdentifier); }));
 		RegisteredTools.Emplace(UICommand, ToolIdentifier);
 	};
 
diff --git a/Engine/Plugins/Experimental/ModelingToolsEditorMode/Source/ModelingToolsEditorMode/Private/ModelingToolsEditorMode.cpp b/Engine/Plugins/Experimental/ModelingToolsEditorMode/Source/ModelingToolsEditorMode/Private/ModelingToolsEditorMode.cpp
--- a/Engine/Plugins/Experimental/ModelingToolsEditorMode/Source/ModelingToolsEditorMode/Private/ModelingToolsEditorMode.cpp
+++ b/Engine/Plugins/Experimental/ModelingToolsEditorMode/Source/ModelingToolsEditorMode/Private/ModelingToolsEditorMode.cpp
@@ -484,7 +484,8 @@ void FModelingToolsEditorMode::Enter()
 		ToolsContext->ToolManager->RegisterToolType(ToolIdentifier, Builder);
 		CommandList->MapAction( UICommand,
 			FExecuteAction::CreateLambda([this, ToolIdentifier]() { ToolsContext->StartTool(ToolIdentifier); }),
-			FCanExecuteAction::CreateLambda([this, ToolIdentifier]() { return ToolsContext->CanStartTool(ToolIdentifier); }));
+			FCanExecuteAction::CreateLambda([this, ToolIdentifier]() { return ToolsContext->CanStartTool(ToolIdentifier); }),
+			FIsActionChecked::CreateLambda([this, ToolIdentifier]() { return ToolsContext->GetActiveToolName() == ToolIdentifier; }));
 
 		RegisteredTools.Emplace(UICommand, ToolIdentifier);
 	};
diff --git a/Engine/Plugins/Experimental/ModelingToolsEditorMode/Source/ModelingToolsEditorMode/Private/ModelingToolsManagerActions.cpp b/Engine/Plugins/Experimental/ModelingToolsEditorMode/Source/ModelingToolsEditorMode/Private/ModelingToolsManagerActions.cpp
--- a/Engine/Plugins/Experimental/ModelingToolsEditorMode/Source/ModelingToolsEditorMode/Private/ModelingToolsManagerActions.cpp
+++ b/Engine/Plugins/Experimental/ModelingToolsEditorMode/Source/ModelingToolsEditorMode/Private/ModelingToolsManagerActions.cpp
@@ -22,90 +22,90 @@ FModelingToolsManagerCommands::FModelingToolsManagerCommands() :
 
 void FModelingToolsManagerCommands::RegisterCommands()
 {
-	UI_COMMAND(BeginAddBoxPrimitiveTool, "Box", "Add Boxes", EUserInterfaceActionType::Button, FInputChord());
-	UI_COMMAND(BeginAddCylinderPrimitiveTool, "Cylinder", "Add Cylinders", EUserInterfaceActionType::Button, FInputChord());
-	UI_COMMAND(BeginAddConePrimitiveTool, "Cone", "Add Cones", EUserInterfaceActionType::Button, FInputChord());
-	UI_COMMAND(BeginAddArrowPrimitiveTool, "Arrow", "Add Arrows", EUserInterfaceActionType::Button, FInputChord());
-	UI_COMMAND(BeginAddRectanglePrimitiveTool, "Rectangle", "Add Rectangles", EUserInterfaceActionType::Button, FInputChord());
-	UI_COMMAND(BeginAddRoundedRectanglePrimitiveTool, "Rounded Rectangle", "Add Rounded Rectangles", EUserInterfaceActionType::Button, FInputChord());
-	UI_COMMAND(BeginAddDiscPrimitiveTool, "Disc", "Add Discs", EUserInterfaceActionType::Button, FInputChord());
-	UI_COMMAND(BeginAddPuncturedDiscPrimitiveTool, "Punctured Disc", "Add PuncturedDiscs", EUserInterfaceActionType::Button, FInputChord());
-	UI_COMMAND(BeginAddTorusPrimitiveTool, "Torus", "Add Torii", EUserInterfaceActionType::Button, FInputChord());
-	UI_COMMAND(BeginAddSphericalBoxPrimitiveTool, "Spherical Box", "Add Box Parametrized Spheres", EUserInterfaceActionType::Button, FInputChord());
-	UI_COMMAND(BeginAddSpherePrimitiveTool, "Sphere", "Add Lat/Long Parametrized Spheres", EUserInterfaceActionType::Button, FInputChord());
-
-	UI_COMMAND(BeginAddPatchTool, "Patch", "Start the Add Patch Tool", EUserInterfaceActionType::Button, FInputChord());
-	UI_COMMAND(BeginRevolveBoundaryTool, "CutRevolve", "Start the CutRevolve Tool, which revolves a section curve created by cutting the selected Mesh", EUserInterfaceActionType::Button, FInputChord());
-	UI_COMMAND(BeginDrawPolygonTool, "Polygon", "Start the Draw Polygon Tool", EUserInterfaceActionType::Button, FInputChord());
-	UI_COMMAND(BeginDrawPolyPathTool, "PolyPath", "Start the Draw PolyPath Tool", EUserInterfaceActionType::Button, FInputChord());
-	UI_COMMAND(BeginDrawAndRevolveTool, "PolyRevolve", "Start the Draw-and-Revolve Tool", EUserInterfaceActionType::Button, FInputChord());
-	UI_COMMAND(BeginShapeSprayTool, "ShapeSpray", "Start the Shape Spray Tool", EUserInterfaceActionType::Button, FInputChord());
-
-	UI_COMMAND(BeginSculptMeshTool, "Sculpt", "Start the Sculpt Mesh Tool", EUserInterfaceActionType::Button, FInputChord());
-	UI_COMMAND(BeginPolyEditTool, "PolyEdit", "Start the PolyEdit Tool", EUserInterfaceActionType::Button, FInputChord());
-	UI_COMMAND(BeginGroupEdgeInsertionTool, "EdgeInsert", "Start the Group Edge Insertion Tool", EUserInterfaceActionType::Button, FInputChord());
-	UI_COMMAND(BeginEdgeLoopInsertionTool, "LoopInsert", "Start the Edge Loop Insertion Tool", EUserInterfaceActionType::Button, FInputChord());
-	UI_COMMAND(BeginTriEditTool, "TriEdit", "Start the Triangle Edit Tool", EUserInterfaceActionType::Button, FInputChord());
-	UI_COMMAND(BeginPolyDeformTool, "PolyDeform", "Start the PolyDeform Tool", EUserInterfaceActionType::Button, FInputChord());
-	UI_COMMAND(BeginSmoothMeshTool, "Smooth", "Start the Smooth Mesh Tool", EUserInterfaceActionType::Button, FInputChord());
-	UI_COMMAND(BeginOffsetMeshTool, "Offset", "Start the Offset Mesh Tool", EUserInterfaceActionType::Button, FInputChord());
-	UI_COMMAND(BeginDisplaceMeshTool, "Displace", "Start the Displace Mesh Tool", EUserInterfaceActionType::Button, FInputChord());
-	UI_COMMAND(BeginMeshSpaceDeformerTool, "SpaceDeform", "Start the Mesh Space Deformer Tool", EUserInterfaceActionType::Button, FInputChord());
-	UI_COMMAND(BeginTransformMeshesTool, "Transform", "Start the Transform Tool", EUserInterfaceActionType::Button, FInputChord());
-	UI_COMMAND(BeginEditPivotTool, "Edit Pivot", "Start the Edit Pivot Tool", EUserInterfaceActionType::Button, FInputChord());
-	UI_COMMAND(BeginBakeTransformTool, "Bake XForm", "Start the Bake Transform Tool", EUserInterfaceActionType::Button, FInputChord());
-	UI_COMMAND(BeginCombineMeshesTool, "Combine", "Combine the selected Mesh Assets into a new Asset", EUserInterfaceActionType::Button, FInputChord());
-	UI_COMMAND(BeginDuplicateMeshesTool, "Duplicate", "Duplicate the selected Mesh Asset into a new Asset", EUserInterfaceActionType::Button, FInputChord());
-	UI_COMMAND(BeginAlignObjectsTool, "Align", "Start the Align Objects Tool", EUserInterfaceActionType::Button, FInputChord());
-
-	UI_COMMAND(BeginRemeshSculptMeshTool, "DynaSculpt", "Start the Liquid Sculpt Mesh Tool", EUserInterfaceActionType::Button, FInputChord());
-	UI_COMMAND(BeginRemeshMeshTool, "Remesh", "Start the Remesh Tool", EUserInterfaceActionType::Button, FInputChord());
-	UI_COMMAND(BeginProjectToTargetTool, "MapToMesh", "Start the Remesh To Target Tool", EUserInterfaceActionType::Button, FInputChord());
-	UI_COMMAND(BeginSimplifyMeshTool, "Simplify", "Start the Simplify Mesh Tool", EUserInterfaceActionType::Button, FInputChord());
-	UI_COMMAND(BeginEditNormalsTool, "Normals", "Start the Edit Normals Tool", EUserInterfaceActionType::Button, FInputChord());
-	UI_COMMAND(BeginEditTangentsTool, "Tangents", "Start the Edit Tangents Tool", EUserInterfaceActionType::Button, FInputChord());
-	UI_COMMAND(BeginRemoveOccludedTrianglesTool, "Jacketing", "Start the Jacketing Tool", EUserInterfaceActionType::Button, FInputChord());
-	UI_COMMAND(BeginUVProjectionTool, "UVProjection", "Start the UV Projection Tool", EUserInterfaceActionType::Button, FInputChord());
-	UI_COMMAND(BeginUVLayoutTool, "UVLayout", "Start the UV Layout Tool", EUserInterfaceActionType::Button, FInputChord());
-	UI_COMMAND(BeginVoxelMergeTool, "VoxMerge", "Start the Voxel Merge Tool", EUserInterfaceActionType::Button, FInputChord());
-	UI_COMMAND(BeginVoxelBooleanTool, "VoxBoolean", "Start the Voxel Boolean Tool", EUserInterfaceActionType::Button, FInputChord());
-	UI_COMMAND(BeginVoxelSolidifyTool, "VoxWrap", "Start the Voxel Shell Tool", EUserInterfaceActionType::Button, FInputChord());
-	UI_COMMAND(BeginVoxelBlendTool, "VoxBlend", "Start the Voxel Blend Tool", EUserInterfaceActionType::Button, FInputChord());
-	UI_COMMAND(BeginVoxelMorphologyTool, "VoxMorphology", "Start the Voxel Morphology Tool", EUserInterfaceActionType::Button, FInputChord());
-	UI_COMMAND(BeginSelfUnionTool, "SelfUnion", "Start the Mesh-Based Self-Union Tool, which will resolve self-intersections", EUserInterfaceActionType::Button, FInputChord());
-	UI_COMMAND(BeginMeshBooleanTool, "Mesh Boolean", "Start the Mesh-Based Boolean Tool", EUserInterfaceActionType::Button, FInputChord());
-	UI_COMMAND(BeginBspConversionTool, "BSP Conversion", "Start the BSP-to-StaticMesh Conversion Tool", EUserInterfaceActionType::Button, FInputChord());
-	UI_COMMAND(BeginMeshToVolumeTool, "Mesh2Vol", "Start the Mesh to Volume Conversion Tool", EUserInterfaceActionType::Button, FInputChord());
-	UI_COMMAND(BeginVolumeToMeshTool, "Vol2Mesh", "Start the Volume to Mesh Conversion Tool", EUserInterfaceActionType::Button, FInputChord());
-	UI_COMMAND(BeginPlaneCutTool, "PlaneCut", "Start the Plane Cut Tool", EUserInterfaceActionType::Button, FInputChord());
-	UI_COMMAND(BeginMirrorTool, "Mirror", "Start the Mirror Tool", EUserInterfaceActionType::Button, FInputChord());
-	UI_COMMAND(BeginHoleFillTool, "HoleFill", "Start the Hole Fill Tool", EUserInterfaceActionType::Button, FInputChord());
-	UI_COMMAND(BeginLatticeDeformerTool, "LatticeDeformer", "Start the Lattice Deformer Tool", EUserInterfaceActionType::Button, FInputChord());
-	UI_COMMAND(BeginPolygonCutTool, "PolyCut", "Start the Polygon Cut Tool", EUserInterfaceActionType::Button, FInputChord());
-	UI_COMMAND(BeginGlobalUVGenerateTool, "AutoUnwrap", "Start the Global UV Unwrap Tool", EUserInterfaceActionType::Button, FInputChord());
-	UI_COMMAND(BeginGroupUVGenerateTool, "GroupUnwrap", "Stat the PolyGroup UV Unwrap Tool", EUserInterfaceActionType::Button, FInputChord());
-	UI_COMMAND(BeginMeshSelectionTool, "Select", "Start the Mesh Selection Tool", EUserInterfaceActionType::Button, FInputChord());
-
-	UI_COMMAND(BeginPhysicsInspectorTool, "PhysInspect", "Inspect Physics Geometry for selected Meshes", EUserInterfaceActionType::Button, FInputChord());
-	UI_COMMAND(BeginSetCollisionGeometryTool, "MeshToCollision", "Convert selected Meshes to Collision Geometry", EUserInterfaceActionType::Button, FInputChord());
-	UI_COMMAND(BeginEditCollisionGeometryTool, "EditPhys", "Edit Collision Geometry for selected Mesh", EUserInterfaceActionType::Button, FInputChord());
-	UI_COMMAND(BeginExtractCollisionGeometryTool, "CollisionToMesh", "Convert Collision Geometry to a Mesh", EUserInterfaceActionType::Button, FInputChord());
-
-	UI_COMMAND(BeginMeshInspectorTool, "Inspector", "Start the Mesh Inspector Tool", EUserInterfaceActionType::Button, FInputChord());
-	UI_COMMAND(BeginWeldEdgesTool, "Weld Edges", "Start the Weld Edges Tool", EUserInterfaceActionType::Button, FInputChord());
-	UI_COMMAND(BeginPolyGroupsTool, "PolyGroups", "Start the PolyGroups Tool", EUserInterfaceActionType::Button, FInputChord());
-	UI_COMMAND(BeginEditMeshMaterialsTool, "Edit Mats", "Start the Material Editing Tool", EUserInterfaceActionType::Button, FInputChord());
-	UI_COMMAND(BeginTransformUVIslandsTool, "Transform UVs", "Start the UV Island Transformation Tool", EUserInterfaceActionType::Button, FInputChord());
-	UI_COMMAND(BeginMeshAttributePaintTool, "Paint Maps", "Start the Paint Maps Tool", EUserInterfaceActionType::Button, FInputChord());
-	UI_COMMAND(BeginAttributeEditorTool, "Edit Attribs", "Start the Attribute Editor Tool", EUserInterfaceActionType::Button, FInputChord());
-	UI_COMMAND(BeginBakeMeshAttributeMapsTool, "Bake Maps", "Start the Map Baking Tool", EUserInterfaceActionType::Button, FInputChord());
-	UI_COMMAND(BeginUVSeamEditTool, "UVSeamEdit", "Start the UV Seam Editing Tool", EUserInterfaceActionType::Button, FInputChord());
-
-	UI_COMMAND(BeginGroomToMeshTool, "HairHelmet", "Start the Hair Helmet Tool", EUserInterfaceActionType::Button, FInputChord());
-	UI_COMMAND(BeginGenerateLODMeshesTool, "GenLODs", "Start the Generate LOD Meshes Tool", EUserInterfaceActionType::Button, FInputChord());
-	UI_COMMAND(BeginGroomCardsToMeshTool, "CardsToMesh", "Start the Hair Cards to Mesh Tool", EUserInterfaceActionType::Button, FInputChord());
-	UI_COMMAND(BeginGroomMeshToCardsTool, "MeshToCards", "Start the Mesh to Hair Cards Tool", EUserInterfaceActionType::Button, FInputChord());
-	UI_COMMAND(BeginGroomCardsEditorTool, "CardsEditor", "Start the Cards Editor Tool", EUserInterfaceActionType::Button, FInputChord());
+	UI_COMMAND(BeginAddBoxPrimitiveTool, "Box", "Add Boxes", EUserInterfaceActionType::ToggleButton, FInputChord());
+	UI_COMMAND(BeginAddCylinderPrimitiveTool, "Cylinder", "Add Cylinders", EUserInterfaceActionType::ToggleButton, FInputChord());
+	UI_COMMAND(BeginAddConePrimitiveTool, "Cone", "Add Cones", EUserInterfaceActionType::ToggleButton, FInputChord());
+	UI_COMMAND(BeginAddArrowPrimitiveTool, "Arrow", "Add Arrows", EUserInterfaceActionType::ToggleButton, FInputChord());
+	UI_COMMAND(BeginAddRectanglePrimitiveTool, "Rectangle", "Add Rectangles", EUserInterfaceActionType::ToggleButton, FInputChord());
+	UI_COMMAND(BeginAddRoundedRectanglePrimitiveTool, "Rounded Rectangle", "Add Rounded Rectangles", EUserInterfaceActionType::ToggleButton, FInputChord());
+	UI_COMMAND(BeginAddDiscPrimitiveTool, "Disc", "Add Discs", EUserInterfaceActionType::ToggleButton, FInputChord());
+	UI_COMMAND(BeginAddPuncturedDiscPrimitiveTool, "Punctured Disc", "Add PuncturedDiscs", EUserInterfaceActionType::ToggleButton, FInputChord());
+	UI_COMMAND(BeginAddTorusPrimitiveTool, "Torus", "Add Torii", EUserInterfaceActionType::ToggleButton, FInputChord());
+	UI_COMMAND(BeginAddSphericalBoxPrimitiveTool, "Spherical Box", "Add Box Parametrized Spheres", EUserInterfaceActionType::ToggleButton, FInputChord());
+	UI_COMMAND(BeginAddSpherePrimitiveTool, "Sphere", "Add Lat/Long Parametrized Spheres", EUserInterfaceActionType::ToggleButton, FInputChord());
+
+	UI_COMMAND(BeginAddPatchTool, "Patch", "Start the Add Patch Tool", EUserInterfaceActionType::ToggleButton, FInputChord());
+	UI_COMMAND(BeginRevolveBoundaryTool, "CutRevolve", "Start the CutRevolve Tool, which revolves a section curve created by cutting the selected Mesh", EUserInterfaceActionType::ToggleButton, FInputChord());
+	UI_COMMAND(BeginDrawPolygonTool, "Polygon", "Start the Draw Polygon Tool", EUserInterfaceActionType::ToggleButton, FInputChord());
+	UI_COMMAND(BeginDrawPolyPathTool, "PolyPath", "Start the Draw PolyPath Tool", EUserInterfaceActionType::ToggleButton, FInputChord());
+	UI_COMMAND(BeginDrawAndRevolveTool, "PolyRevolve", "Start the Draw-and-Revolve Tool", EUserInterfaceActionType::ToggleButton, FInputChord());
+	UI_COMMAND(BeginShapeSprayTool, "ShapeSpray", "Start the Shape Spray Tool", EUserInterfaceActionType::ToggleButton, FInputChord());
+
+	UI_COMMAND(BeginSculptMeshTool, "Sculpt", "Start the Sculpt Mesh Tool", EUserInterfaceActionType::ToggleButton, FInputChord());
+	UI_COMMAND(BeginPolyEditTool, "PolyEdit", "Start the PolyEdit Tool", EUserInterfaceActionType::ToggleButton, FInputChord());
+	UI_COMMAND(BeginGroupEdgeInsertionTool, "EdgeInsert", "Start the Group Edge Insertion Tool", EUserInterfaceActionType::ToggleButton, FInputChord());
+	UI_COMMAND(BeginEdgeLoopInsertionTool, "LoopInsert", "Start the Edge Loop Insertion Tool", EUserInterfaceActionType::ToggleButton, FInputChord());
+	UI_COMMAND(BeginTriEditTool, "TriEdit", "Start the Triangle Edit Tool", EUserInterfaceActionType::ToggleButton, FInputChord());
+	UI_COMMAND(BeginPolyDeformTool, "PolyDeform", "Start the PolyDeform Tool", EUserInterfaceActionType::ToggleButton, FInputChord());
+	UI_COMMAND(BeginSmoothMeshTool, "Smooth", "Start the Smooth Mesh Tool", EUserInterfaceActionType::ToggleButton, FInputChord());
+	UI_COMMAND(BeginOffsetMeshTool, "Offset", "Start the Offset Mesh Tool", EUserInterfaceActionType::ToggleButton, FInputChord());
+	UI_COMMAND(BeginDisplaceMeshTool, "Displace", "Start the Displace Mesh Tool", EUserInterfaceActionType::ToggleButton, FInputChord());
+	UI_COMMAND(BeginMeshSpaceDeformerTool, "SpaceDeform", "Start the Mesh Space Deformer Tool", EUserInterfaceActionType::ToggleButton, FInputChord());
+	UI_COMMAND(BeginTransformMeshesTool, "Transform", "Start the Transform Tool", EUserInterfaceActionType::ToggleButton, FInputChord());
+	UI_COMMAND(BeginEditPivotTool, "Edit Pivot", "Start the Edit Pivot Tool", EUserInterfaceActionType::ToggleButton, FInputChord());
+	UI_COMMAND(BeginBakeTransformTool, "Bake XForm", "Start the Bake Transform Tool", EUserInterfaceActionType::ToggleButton, FInputChord());
+	UI_COMMAND(BeginCombineMeshesTool, "Combine", "Combine the selected Mesh Assets into a new Asset", EUserInterfaceActionType::ToggleButton, FInputChord());
+	UI_COMMAND(BeginDuplicateMeshesTool, "Duplicate", "Duplicate the selected Mesh Asset into a new Asset", EUserInterfaceActionType::ToggleButton, FInputChord());
+	UI_COMMAND(BeginAlignObjectsTool, "Align", "Start the Align Objects Tool", EUserInterfaceActionType::ToggleButton, FInputChord());
+
+	UI_COMMAND(BeginRemeshSculptMeshTool, "DynaSculpt", "Start the Liquid Sculpt Mesh Tool", EUserInterfaceActionType::ToggleButton, FInputChord());
+	UI_COMMAND(BeginRemeshMeshTool, "Remesh", "Start the Remesh Tool", EUserInterfaceActionType::ToggleButton, FInputChord());
+	UI_COMMAND(BeginProjectToTargetTool, "MapToMesh", "Start the Remesh To Target Tool", EUserInterfaceActionType::ToggleButton, FInputChord());
+	UI_COMMAND(BeginSimplifyMeshTool, "Simplify", "Start the Simplify Mesh Tool", EUserInterfaceActionType::ToggleButton, FInputChord());
+	UI_COMMAND(BeginEditNormalsTool, "Normals", "Start the Edit Normals Tool", EUserInterfaceActionType::ToggleButton, FInputChord());
+	UI_COMMAND(BeginEditTangentsTool, "Tangents", "Start the Edit Tangents Tool", EUserInterfaceActionType::ToggleButton, FInputChord());
+	UI_COMMAND(BeginRemoveOccludedTrianglesTool, "Jacketing", "Start the Jacketing Tool", EUserInterfaceActionType::ToggleButton, FInputChord());
+	UI_COMMAND(BeginUVProjectionTool, "UVProjection", "Start the UV Projection Tool", EUserInterfaceActionType::ToggleButton, FInputChord());
+	UI_COMMAND(BeginUVLayoutTool, "UVLayout", "Start the UV Layout Tool", EUserInterfaceActionType::ToggleButton, FInputChord());
+	UI_COMMAND(BeginVoxelMergeTool, "VoxMerge", "Start the Voxel Merge Tool", EUserInterfaceActionType::ToggleButton, FInputChord());
+	UI_COMMAND(BeginVoxelBooleanTool, "VoxBoolean", "Start the Voxel Boolean Tool", EUserInterfaceActionType::ToggleButton, FInputChord());
+	UI_COMMAND(BeginVoxelSolidifyTool, "VoxWrap", "Start the Voxel Shell Tool", EUserInterfaceActionType::ToggleButton, FInputChord());
+	UI_COMMAND(BeginVoxelBlendTool, "VoxBlend", "Start the Voxel Blend Tool", EUserInterfaceActionType::ToggleButton, FInputChord());
+	UI_COMMAND(BeginVoxelMorphologyTool, "VoxMorphology", "Start the Voxel Morphology Tool", EUserInterfaceActionType::ToggleButton, FInputChord());
+	UI_COMMAND(BeginSelfUnionTool, "SelfUnion", "Start the Mesh-Based Self-Union Tool, which will resolve self-intersections", EUserInterfaceActionType::ToggleButton, FInputChord());
+	UI_COMMAND(BeginMeshBooleanTool, "Mesh Boolean", "Start the Mesh-Based Boolean Tool", EUserInterfaceActionType::ToggleButton, FInputChord());
+	UI_COMMAND(BeginBspConversionTool, "BSP Conversion", "Start the BSP-to-StaticMesh Conversion Tool", EUserInterfaceActionType::ToggleButton, FInputChord());
+	UI_COMMAND(BeginMeshToVolumeTool, "Mesh2Vol", "Start the Mesh to Volume Conversion Tool", EUserInterfaceActionType::ToggleButton, FInputChord());
+	UI_COMMAND(BeginVolumeToMeshTool, "Vol2Mesh", "Start the Volume to Mesh Conversion Tool", EUserInterfaceActionType::ToggleButton, FInputChord());
+	UI_COMMAND(BeginPlaneCutTool, "PlaneCut", "Start the Plane Cut Tool", EUserInterfaceActionType::ToggleButton, FInputChord());
+	UI_COMMAND(BeginMirrorTool, "Mirror", "Start the Mirror Tool", EUserInterfaceActionType::ToggleButton, FInputChord());
+	UI_COMMAND(BeginHoleFillTool, "HoleFill", "Start the Hole Fill Tool", EUserInterfaceActionType::ToggleButton, FInputChord());
+	UI_COMMAND(BeginLatticeDeformerTool, "LatticeDeformer", "Start the Lattice Deformer Tool", EUserInterfaceActionType::ToggleButton, FInputChord());
+	UI_COMMAND(BeginPolygonCutTool, "PolyCut", "Start the Polygon Cut Tool", EUserInterfaceActionType::ToggleButton, FInputChord());
+	UI_COMMAND(BeginGlobalUVGenerateTool, "AutoUnwrap", "Start the Global UV Unwrap Tool", EUserInterfaceActionType::ToggleButton, FInputChord());
+	UI_COMMAND(BeginGroupUVGenerateTool, "GroupUnwrap", "Stat the PolyGroup UV Unwrap Tool", EUserInterfaceActionType::ToggleButton, FInputChord());
+	UI_COMMAND(BeginMeshSelectionTool, "Select", "Start the Mesh Selection Tool", EUserInterfaceActionType::ToggleButton, FInputChord());
+
+	UI_COMMAND(BeginPhysicsInspectorTool, "PhysInspect", "Inspect Physics Geometry for selected Meshes", EUserInterfaceActionType::ToggleButton, FInputChord());
+	UI_COMMAND(BeginSetCollisionGeometryTool, "MeshToCollision", "Convert selected Meshes to Collision Geometry", EUserInterfaceActionType::ToggleButton, FInputChord());
+	UI_COMMAND(BeginEditCollisionGeometryTool, "EditPhys", "Edit Collision Geometry for selected Mesh", EUserInterfaceActionType::ToggleButton, FInputChord());
+	UI_COMMAND(BeginExtractCollisionGeometryTool, "CollisionToMesh", "Convert Collision Geometry to a Mesh", EUserInterfaceActionType::ToggleButton, FInputChord());
+
+	UI_COMMAND(BeginMeshInspectorTool, "Inspector", "Start the Mesh Inspector Tool", EUserInterfaceActionType::ToggleButton, FInputChord());
+	UI_COMMAND(BeginWeldEdgesTool, "Weld Edges", "Start the Weld Edges Tool", EUserInterfaceActionType::ToggleButton, FInputChord());
+	UI_COMMAND(BeginPolyGroupsTool, "PolyGroups", "Start the PolyGroups Tool", EUserInterfaceActionType::ToggleButton, FInputChord());
+	UI_COMMAND(BeginEditMeshMaterialsTool, "Edit Mats", "Start the Material Editing Tool", EUserInterfaceActionType::ToggleButton, FInputChord());
+	UI_COMMAND(BeginTransformUVIslandsTool, "Transform UVs", "Start the UV Island Transformation Tool", EUserInterfaceActionType::ToggleButton, FInputChord());
+	UI_COMMAND(BeginMeshAttributePaintTool, "Paint Maps", "Start the Paint Maps Tool", EUserInterfaceActionType::ToggleButton, FInputChord());
+	UI_COMMAND(BeginAttributeEditorTool, "Edit Attribs", "Start the Attribute Editor Tool", EUserInterfaceActionType::ToggleButton, FInputChord());
+	UI_COMMAND(BeginBakeMeshAttributeMapsTool, "Bake Maps", "Start the Map Baking Tool", EUserInterfaceActionType::ToggleButton, FInputChord());
+	UI_COMMAND(BeginUVSeamEditTool, "UVSeamEdit", "Start the UV Seam Editing Tool", EUserInterfaceActionType::ToggleButton, FInputChord());
+
+	UI_COMMAND(BeginGroomToMeshTool, "HairHelmet", "Start the Hair Helmet Tool", EUserInterfaceActionType::ToggleButton, FInputChord());
+	UI_COMMAND(BeginGenerateLODMeshesTool, "GenLODs", "Start the Generate LOD Meshes Tool", EUserInterfaceActionType::ToggleButton, FInputChord());
+	UI_COMMAND(BeginGroomCardsToMeshTool, "CardsToMesh", "Start the Hair Cards to Mesh Tool", EUserInterfaceActionType::ToggleButton, FInputChord());
+	UI_COMMAND(BeginGroomMeshToCardsTool, "MeshToCards", "Start the Mesh to Hair Cards Tool", EUserInterfaceActionType::ToggleButton, FInputChord());
+	UI_COMMAND(BeginGroomCardsEditorTool, "CardsEditor", "Start the Cards Editor Tool", EUserInterfaceActionType::ToggleButton, FInputChord());
 
 	UI_COMMAND(AcceptActiveTool, "Accept", "Accept the active tool", EUserInterfaceActionType::Button, FInputChord());
 	UI_COMMAND(CancelActiveTool, "Cancel", "Cancel the active tool", EUserInterfaceActionType::Button, FInputChord());
diff --git a/Engine/Source/Editor/Experimental/EditorInteractiveToolsFramework/Private/EdModeInteractiveToolsContext.cpp b/Engine/Source/Editor/Experimental/EditorInteractiveToolsFramework/Private/EdModeInteractiveToolsContext.cpp
--- a/Engine/Source/Editor/Experimental/EditorInteractiveToolsFramework/Private/EdModeInteractiveToolsContext.cpp
+++ b/Engine/Source/Editor/Experimental/EditorInteractiveToolsFramework/Private/EdModeInteractiveToolsContext.cpp
@@ -1194,6 +1194,16 @@ bool UEdModeInteractiveToolsContext::CanStartTool(const FString ToolTypeIdentifi
 	return UInteractiveToolsContext::CanStartTool(EToolSide::Mouse, ToolTypeIdentifier);
 }
 
+bool UEdModeInteractiveToolsContext::HasActiveTool() const
+{
+	return UInteractiveToolsContext::HasActiveTool(EToolSide::Mouse);
+}
+
+FString UEdModeInteractiveToolsContext::GetActiveToolName() const
+{
+	return UInteractiveToolsContext::GetActiveToolName(EToolSide::Mouse);
+}
+
 bool UEdModeInteractiveToolsContext::ActiveToolHasAccept() const
 {
 	return  UInteractiveToolsContext::ActiveToolHasAccept(EToolSide::Mouse);
diff --git a/Engine/Source/Editor/Experimental/EditorInteractiveToolsFramework/Public/EdModeInteractiveToolsContext.h b/Engine/Source/Editor/Experimental/EditorInteractiveToolsFramework/Public/EdModeInteractiveToolsContext.h
--- a/Engine/Source/Editor/Experimental/EditorInteractiveToolsFramework/Public/EdModeInteractiveToolsContext.h
+++ b/Engine/Source/Editor/Experimental/EditorInteractiveToolsFramework/Public/EdModeInteractiveToolsContext.h
@@ -82,6 +82,8 @@ public:
 	//
 
 	virtual bool CanStartTool(const FString ToolTypeIdentifier) const;
+	virtual bool HasActiveTool() const;
+	virtual FString GetActiveToolName() const;
 	virtual bool ActiveToolHasAccept() const;
 	virtual bool CanAcceptActiveTool() const;
 	virtual bool CanCancelActiveTool() const;
diff --git a/Engine/Source/Runtime/Experimental/InteractiveToolsFramework/Private/InteractiveToolsContext.cpp b/Engine/Source/Runtime/Experimental/InteractiveToolsFramework/Private/InteractiveToolsContext.cpp
--- a/Engine/Source/Runtime/Experimental/InteractiveToolsFramework/Private/InteractiveToolsContext.cpp
+++ b/Engine/Source/Runtime/Experimental/InteractiveToolsFramework/Private/InteractiveToolsContext.cpp
@@ -65,8 +65,17 @@ void UInteractiveToolsContext::DeactivateAllActiveTools()
 
 bool UInteractiveToolsContext::CanStartTool(EToolSide WhichSide, const FString& ToolTypeIdentifier) const
 {
-	return (ToolManager->HasActiveTool(WhichSide) == false) &&
-		(ToolManager->CanActivateTool(WhichSide, ToolTypeIdentifier) == true);
+	return ToolManager->CanActivateTool(WhichSide, ToolTypeIdentifier);
+}
+
+bool UInteractiveToolsContext::HasActiveTool(EToolSide WhichSide) const
+{
+	return ToolManager->HasActiveTool(WhichSide);
+}
+
+FString UInteractiveToolsContext::GetActiveToolName(EToolSide WhichSide) const
+{
+	return ToolManager->GetActiveToolName(WhichSide);
 }
 
 bool UInteractiveToolsContext::ActiveToolHasAccept(EToolSide WhichSide) const
diff --git a/Engine/Source/Runtime/Experimental/InteractiveToolsFramework/Public/InteractiveToolsContext.h b/Engine/Source/Runtime/Experimental/InteractiveToolsFramework/Public/InteractiveToolsContext.h
--- a/Engine/Source/Runtime/Experimental/InteractiveToolsFramework/Public/InteractiveToolsContext.h
+++ b/Engine/Source/Runtime/Experimental/InteractiveToolsFramework/Public/InteractiveToolsContext.h
@@ -34,6 +34,8 @@ public:
 	virtual void DeactivateAllActiveTools();
 
 	bool CanStartTool(EToolSide WhichSide, const FString& ToolTypeIdentifier) const;
+	bool HasActiveTool(EToolSide WhichSide) const;
+	FString GetActiveToolName(EToolSide WhichSide) const;
 	bool ActiveToolHasAccept(EToolSide WhichSide) const;
 	bool CanAcceptActiveTool(EToolSide WhichSide) const;
 	bool CanCancelActiveTool(EToolSide WhichSide) const;
