ModelingMode: add PolyModel tab with experimental refactor of PolyEdit selection and operation tools
#rnx

diff --git a/Engine/Plugins/Experimental/MeshModelingToolset/Source/MeshModelingTools/Private/EditMeshPolygonsTool.cpp b/Engine/Plugins/Experimental/MeshModelingToolset/Source/MeshModelingTools/Private/EditMeshPolygonsTool.cpp
--- a/Engine/Plugins/Experimental/MeshModelingToolset/Source/MeshModelingTools/Private/EditMeshPolygonsTool.cpp
+++ b/Engine/Plugins/Experimental/MeshModelingToolset/Source/MeshModelingTools/Private/EditMeshPolygonsTool.cpp
@@ -63,6 +63,73 @@ UMeshSurfacePointTool* UEditMeshPolygonsToolBuilder::CreateNewTool(const FToolBu
 
 
 
+
+
+void UEditMeshPolygonsActionModeToolBuilder::InitializeNewTool(UMeshSurfacePointTool* ToolIn, const FToolBuilderState& SceneState) const
+{
+	UEditMeshPolygonsToolBuilder::InitializeNewTool(ToolIn, SceneState);
+
+	UEditMeshPolygonsTool* Tool = Cast<UEditMeshPolygonsTool>(ToolIn);
+	EEditMeshPolygonsToolActions UseAction = StartupAction;
+	Tool->PostSetupFunction = [UseAction](UEditMeshPolygonsTool* PolyTool)
+	{
+		PolyTool->SetToSelectionModeInterface();
+		PolyTool->RequestAction(UseAction);
+	};
+}
+
+
+void UEditMeshPolygonsSelectionModeToolBuilder::InitializeNewTool(UMeshSurfacePointTool* ToolIn, const FToolBuilderState& SceneState) const
+{
+	UEditMeshPolygonsToolBuilder::InitializeNewTool(ToolIn, SceneState);
+
+	UEditMeshPolygonsTool* Tool = Cast<UEditMeshPolygonsTool>(ToolIn);
+	EEditMeshPolygonsToolSelectionMode UseMode = SelectionMode;
+	Tool->PostSetupFunction = [UseMode](UEditMeshPolygonsTool* PolyTool)
+	{
+		PolyTool->SetToSelectionModeInterface();
+
+		UPolygonSelectionMechanic* SelectionMechanic = PolyTool->SelectionMechanic;
+		UPolygonSelectionMechanicProperties* SelectionProps = SelectionMechanic->Properties;
+		SelectionProps->bSelectFaces = SelectionProps->bSelectEdges = SelectionProps->bSelectVertices = false;
+		SelectionProps->bSelectEdgeLoops = SelectionProps->bSelectEdgeRings = false;
+
+		switch (UseMode)
+		{
+		default:
+		case EEditMeshPolygonsToolSelectionMode::Faces:
+			SelectionProps->bSelectFaces = true;
+			break;
+		case EEditMeshPolygonsToolSelectionMode::Edges:
+			SelectionProps->bSelectEdges = true;
+			break;
+		case EEditMeshPolygonsToolSelectionMode::Vertices:
+			SelectionProps->bSelectVertices = true;
+			break;
+		case EEditMeshPolygonsToolSelectionMode::Loops:
+			SelectionProps->bSelectEdges = true;
+			SelectionProps->bSelectEdgeLoops = true;
+			break;
+		case EEditMeshPolygonsToolSelectionMode::Rings:
+			SelectionProps->bSelectEdges = true;
+			SelectionProps->bSelectEdgeRings = true;
+			break;
+		case EEditMeshPolygonsToolSelectionMode::FacesEdgesVertices:
+			SelectionProps->bSelectFaces = SelectionProps->bSelectEdges = SelectionProps->bSelectVertices = true;
+			break;
+		}
+	};
+}
+
+void UEditMeshPolygonsTool::SetToSelectionModeInterface()
+{
+	if (EditActions) SetToolPropertySourceEnabled(EditActions, false);
+	if (EditEdgeActions) SetToolPropertySourceEnabled(EditEdgeActions, false);
+	if (EditUVActions) SetToolPropertySourceEnabled(EditUVActions, false);
+}
+
+
+
 void UEditMeshPolygonsToolActionPropertySet::PostAction(EEditMeshPolygonsToolActions Action)
 {
 	if (ParentTool.IsValid())
@@ -272,6 +339,11 @@ void UEditMeshPolygonsTool::Setup()
 	{
 		GetToolManager()->DisplayMessage( LOCTEXT("NoGroupsWarning", "This object has a single PolyGroup. Use the PolyGroups or Select Tool to assign PolyGroups."), EToolMessageLevel::UserWarning);
 	}
+
+	if (PostSetupFunction)
+	{
+		PostSetupFunction(this);
+	}
 }
 
 bool UEditMeshPolygonsTool::IsStoredToolSelectionUsable(const UGroupTopologyStorableSelection* StoredSelection)
diff --git a/Engine/Plugins/Experimental/MeshModelingToolset/Source/MeshModelingTools/Public/EditMeshPolygonsTool.h b/Engine/Plugins/Experimental/MeshModelingToolset/Source/MeshModelingTools/Public/EditMeshPolygonsTool.h
--- a/Engine/Plugins/Experimental/MeshModelingToolset/Source/MeshModelingTools/Public/EditMeshPolygonsTool.h
+++ b/Engine/Plugins/Experimental/MeshModelingToolset/Source/MeshModelingTools/Public/EditMeshPolygonsTool.h
@@ -112,6 +112,46 @@ enum class EEditMeshPolygonsToolActions
 
 
 
+UCLASS()
+class MESHMODELINGTOOLS_API UEditMeshPolygonsActionModeToolBuilder : public UEditMeshPolygonsToolBuilder
+{
+	GENERATED_BODY()
+public:
+	EEditMeshPolygonsToolActions StartupAction = EEditMeshPolygonsToolActions::Extrude;
+
+	virtual void InitializeNewTool(UMeshSurfacePointTool* Tool, const FToolBuilderState& SceneState) const override;
+};
+
+
+
+
+UENUM()
+enum class EEditMeshPolygonsToolSelectionMode
+{
+	Faces,
+	Edges,
+	Vertices,
+	Loops,
+	Rings,
+	FacesEdgesVertices
+};
+
+UCLASS()
+class MESHMODELINGTOOLS_API UEditMeshPolygonsSelectionModeToolBuilder : public UEditMeshPolygonsToolBuilder
+{
+	GENERATED_BODY()
+public:
+	EEditMeshPolygonsToolSelectionMode SelectionMode = EEditMeshPolygonsToolSelectionMode::Faces;
+
+	virtual void InitializeNewTool(UMeshSurfacePointTool* Tool, const FToolBuilderState& SceneState) const override;
+};
+
+
+
+
+
+
+
 
 UCLASS()
 class MESHMODELINGTOOLS_API UEditMeshPolygonsToolActionPropertySet : public UInteractiveToolPropertySet
@@ -690,6 +730,14 @@ protected:
 	friend class FEditPolygonsTopologyPreEditChange;
 	friend class FEditPolygonsTopologyPostEditChange;
 	friend class FBeginInteractivePolyEditChange;
+
+
+
+	// custom setup support
+	friend class UEditMeshPolygonsSelectionModeToolBuilder;
+	friend class UEditMeshPolygonsActionModeToolBuilder;
+	TUniqueFunction<void(UEditMeshPolygonsTool*)> PostSetupFunction;
+	void SetToSelectionModeInterface();
 };
 
 
diff --git a/Engine/Plugins/Experimental/ModelingToolsEditorMode/Source/ModelingToolsEditorMode/Private/ModelingToolsEditorMode.cpp b/Engine/Plugins/Experimental/ModelingToolsEditorMode/Source/ModelingToolsEditorMode/Private/ModelingToolsEditorMode.cpp
--- a/Engine/Plugins/Experimental/ModelingToolsEditorMode/Source/ModelingToolsEditorMode/Private/ModelingToolsEditorMode.cpp
+++ b/Engine/Plugins/Experimental/ModelingToolsEditorMode/Source/ModelingToolsEditorMode/Private/ModelingToolsEditorMode.cpp
@@ -772,6 +772,31 @@ void FModelingToolsEditorMode::Enter()
 	RegisterToolFunc(ToolManagerCommands.BeginGenerateLODMeshesTool, TEXT("GenerateLODMeshesTool"), GenerateLODMeshesToolBuilder);
 
 
+	// PolyModeling tools
+	auto RegisterPolyModelSelectTool = [&](EEditMeshPolygonsToolSelectionMode SelectionMode, TSharedPtr<FUICommandInfo> UICommand, FString StringName)
+	{
+		UEditMeshPolygonsSelectionModeToolBuilder* SelectionModeBuilder = NewObject<UEditMeshPolygonsSelectionModeToolBuilder>();
+		SelectionModeBuilder->SelectionMode = SelectionMode;
+		RegisterToolFunc(UICommand, StringName, SelectionModeBuilder);
+	};
+	RegisterPolyModelSelectTool(EEditMeshPolygonsToolSelectionMode::Faces, ToolManagerCommands.BeginPolyModelTool_FaceSelect, TEXT("PolyEdit_FaceSelect"));
+	RegisterPolyModelSelectTool(EEditMeshPolygonsToolSelectionMode::Edges, ToolManagerCommands.BeginPolyModelTool_EdgeSelect, TEXT("PolyEdit_EdgeSelect"));
+	RegisterPolyModelSelectTool(EEditMeshPolygonsToolSelectionMode::Vertices, ToolManagerCommands.BeginPolyModelTool_VertexSelect, TEXT("PolyEdit_VertexSelect"));
+	RegisterPolyModelSelectTool(EEditMeshPolygonsToolSelectionMode::Loops, ToolManagerCommands.BeginPolyModelTool_LoopSelect, TEXT("PolyEdit_LoopSelect"));
+	RegisterPolyModelSelectTool(EEditMeshPolygonsToolSelectionMode::Rings, ToolManagerCommands.BeginPolyModelTool_RingSelect, TEXT("PolyEdit_RingSelect"));
+	RegisterPolyModelSelectTool(EEditMeshPolygonsToolSelectionMode::FacesEdgesVertices, ToolManagerCommands.BeginPolyModelTool_AllSelect, TEXT("PolyEdit_AllSelect"));
+
+	auto RegisterPolyModelActionTool = [&](EEditMeshPolygonsToolActions Action, TSharedPtr<FUICommandInfo> UICommand, FString StringName)
+	{
+		UEditMeshPolygonsActionModeToolBuilder* ActionModeBuilder = NewObject<UEditMeshPolygonsActionModeToolBuilder>();
+		ActionModeBuilder->StartupAction = Action;
+		RegisterToolFunc(UICommand, StringName, ActionModeBuilder);
+	};
+	RegisterPolyModelActionTool(EEditMeshPolygonsToolActions::Extrude, ToolManagerCommands.BeginPolyModelTool_Extrude, TEXT("PolyEdit_Extrude"));
+	RegisterPolyModelActionTool(EEditMeshPolygonsToolActions::Offset, ToolManagerCommands.BeginPolyModelTool_Offset, TEXT("PolyEdit_Offset"));
+	RegisterPolyModelActionTool(EEditMeshPolygonsToolActions::Inset, ToolManagerCommands.BeginPolyModelTool_Inset, TEXT("PolyEdit_Inset"));
+	RegisterPolyModelActionTool(EEditMeshPolygonsToolActions::Outset, ToolManagerCommands.BeginPolyModelTool_Outset, TEXT("PolyEdit_Outset"));
+	RegisterPolyModelActionTool(EEditMeshPolygonsToolActions::CutFaces, ToolManagerCommands.BeginPolyModelTool_CutFaces, TEXT("PolyEdit_CutFaces"));
 
 	ToolsContext->ToolManager->SelectActiveToolType(EToolSide::Left, TEXT("DynaSculptTool"));
 
diff --git a/Engine/Plugins/Experimental/ModelingToolsEditorMode/Source/ModelingToolsEditorMode/Private/ModelingToolsEditorModeToolkit.cpp b/Engine/Plugins/Experimental/ModelingToolsEditorMode/Source/ModelingToolsEditorMode/Private/ModelingToolsEditorModeToolkit.cpp
--- a/Engine/Plugins/Experimental/ModelingToolsEditorMode/Source/ModelingToolsEditorMode/Private/ModelingToolsEditorModeToolkit.cpp
+++ b/Engine/Plugins/Experimental/ModelingToolsEditorMode/Source/ModelingToolsEditorMode/Private/ModelingToolsEditorModeToolkit.cpp
@@ -26,8 +26,13 @@
 // if set to 1, then on mode initialization we include buttons for prototype modeling tools
 static TAutoConsoleVariable<int32> CVarEnablePrototypeModelingTools(
 	TEXT("modeling.EnablePrototypes"),
-	0,
+	1,
 	TEXT("Enable unsupported Experimental prototype Modeling Tools"));
+static TAutoConsoleVariable<int32> CVarEnablePolyModeling(
+	TEXT("modeling.EnablePolyModel"),
+	1,
+	TEXT("Enable prototype PolyModeling tab"));
+
 
 
 FModelingToolsEditorModeToolkit::FModelingToolsEditorModeToolkit()
@@ -299,16 +304,23 @@ static const FName DeformTabName(TEXT("Deform"));
 static const FName VolumesTabName(TEXT("Volumes"));
 static const FName PrototypesTabName(TEXT("Prototypes"));
 static const FName HairTabName(TEXT("Hair"));
+static const FName PolyModelTabName(TEXT("PolyModel"));
 
 
 const TArray<FName> FModelingToolsEditorModeToolkit::PaletteNames_Standard = { PrimitiveTabName, CreateTabName, TransformTabName, DeformTabName, PolyGroupsTabName, TrianglesTabName, UVNormalTabName, VolumesTabName, HairTabName };
 const TArray<FName> FModelingToolsEditorModeToolkit::PaletteNames_Experimental = { PrimitiveTabName, CreateTabName, TransformTabName, DeformTabName, PolyGroupsTabName, TrianglesTabName, UVNormalTabName, VolumesTabName, HairTabName, PrototypesTabName };
 
 
-void FModelingToolsEditorModeToolkit::GetToolPaletteNames(TArray<FName>& InPaletteName) const
+void FModelingToolsEditorModeToolkit::GetToolPaletteNames(TArray<FName>& PaletteNames) const
 {
 	bool bEnablePrototypes = (CVarEnablePrototypeModelingTools.GetValueOnGameThread() > 0);
-	InPaletteName = (bEnablePrototypes) ? PaletteNames_Experimental : PaletteNames_Standard;
+	PaletteNames = (bEnablePrototypes) ? PaletteNames_Experimental : PaletteNames_Standard;
+
+	bool bEnablePolyModel = (CVarEnablePolyModeling.GetValueOnGameThread() > 0);
+	if (bEnablePolyModel)
+	{
+		PaletteNames.Add(PolyModelTabName);
+	}
 }
 
 
@@ -320,175 +332,14 @@ FText FModelingToolsEditorModeToolkit::GetToolPaletteDisplayName(FName Palette)
 void FModelingToolsEditorModeToolkit::BuildToolPalette(FName PaletteIndex, class FToolBarBuilder& ToolbarBuilder) 
 {
 	const FModelingToolsManagerCommands& Commands = FModelingToolsManagerCommands::Get();
-	bool bEnablePrototypes = (CVarEnablePrototypeModelingTools.GetValueOnGameThread() > 0);
-	if (bEnablePrototypes)
-	{
-		BuildToolPalette_Experimental(PaletteIndex, ToolbarBuilder);
-	}
-	else
-	{
-		BuildToolPalette_Standard(PaletteIndex, ToolbarBuilder);
-	}
+	BuildToolPalette_Experimental(PaletteIndex, ToolbarBuilder);
 }
 
 
 
 void FModelingToolsEditorModeToolkit::BuildToolPalette_Standard(FName PaletteIndex, class FToolBarBuilder& ToolbarBuilder)
 {
-	const FModelingToolsManagerCommands& Commands = FModelingToolsManagerCommands::Get();
-
-	ToolbarBuilder.AddToolBarButton(Commands.AcceptActiveTool);
-	ToolbarBuilder.AddToolBarButton(Commands.CancelActiveTool);
-	ToolbarBuilder.AddToolBarButton(Commands.CompleteActiveTool);
-
-	ToolbarBuilder.AddSeparator();
-
-
-	if (PaletteIndex == PrimitiveTabName)
-	{
-		ToolbarBuilder.AddToolBarButton(Commands.BeginAddBoxPrimitiveTool);
-		ToolbarBuilder.AddToolBarButton(Commands.BeginAddCylinderPrimitiveTool);
-		ToolbarBuilder.AddToolBarButton(Commands.BeginAddConePrimitiveTool);
-		ToolbarBuilder.AddToolBarButton(Commands.BeginAddArrowPrimitiveTool);
-		ToolbarBuilder.AddToolBarButton(Commands.BeginAddRectanglePrimitiveTool);
-		ToolbarBuilder.AddToolBarButton(Commands.BeginAddRoundedRectanglePrimitiveTool);
-		ToolbarBuilder.AddToolBarButton(Commands.BeginAddDiscPrimitiveTool);
-		ToolbarBuilder.AddToolBarButton(Commands.BeginAddPuncturedDiscPrimitiveTool);
-		ToolbarBuilder.AddToolBarButton(Commands.BeginAddTorusPrimitiveTool);
-		ToolbarBuilder.AddToolBarButton(Commands.BeginAddSpherePrimitiveTool);
-		ToolbarBuilder.AddToolBarButton(Commands.BeginAddSphericalBoxPrimitiveTool);
-	}
-	else if (PaletteIndex == CreateTabName)
-	{
-		ToolbarBuilder.AddToolBarButton(Commands.BeginTransformMeshesTool);
-		ToolbarBuilder.AddToolBarButton(Commands.BeginMeshSelectionTool);
-		ToolbarBuilder.AddSeparator();
-		ToolbarBuilder.AddToolBarButton(Commands.BeginDrawPolygonTool);
-		ToolbarBuilder.AddToolBarButton(Commands.BeginDrawPolyPathTool);
-		ToolbarBuilder.AddToolBarButton(Commands.BeginDrawAndRevolveTool);
-		ToolbarBuilder.AddToolBarButton(Commands.BeginRevolveBoundaryTool);
-		ToolbarBuilder.AddSeparator();
-		ToolbarBuilder.AddToolBarButton(Commands.BeginMeshBooleanTool);
-		ToolbarBuilder.AddToolBarButton(Commands.BeginSelfUnionTool);
-		ToolbarBuilder.AddSeparator();
-		ToolbarBuilder.AddToolBarButton(Commands.BeginVoxelSolidifyTool);
-		ToolbarBuilder.AddToolBarButton(Commands.BeginVoxelBlendTool);
-		ToolbarBuilder.AddToolBarButton(Commands.BeginVoxelMorphologyTool);
-		ToolbarBuilder.AddToolBarButton(Commands.BeginVoxelBooleanTool);
-		ToolbarBuilder.AddToolBarButton(Commands.BeginVoxelMergeTool);
-	}
-	else if (PaletteIndex == TransformTabName)
-	{
-		ToolbarBuilder.AddToolBarButton(Commands.BeginTransformMeshesTool);
-		ToolbarBuilder.AddToolBarButton(Commands.BeginMeshSelectionTool);
-		ToolbarBuilder.AddSeparator();
-		ToolbarBuilder.AddToolBarButton(Commands.BeginAlignObjectsTool);
-		ToolbarBuilder.AddToolBarButton(Commands.BeginEditPivotTool);
-		ToolbarBuilder.AddToolBarButton(Commands.BeginBakeTransformTool);
-		ToolbarBuilder.AddToolBarButton(Commands.BeginCombineMeshesTool);
-		ToolbarBuilder.AddToolBarButton(Commands.BeginDuplicateMeshesTool);
-		ToolbarBuilder.AddSeparator();
-		ToolbarBuilder.AddToolBarButton(Commands.BeginMeshInspectorTool);
-	}
-	else if (PaletteIndex == DeformTabName)
-	{
-		ToolbarBuilder.AddToolBarButton(Commands.BeginTransformMeshesTool);
-		ToolbarBuilder.AddToolBarButton(Commands.BeginMeshSelectionTool);
-		ToolbarBuilder.AddSeparator();
-		ToolbarBuilder.AddToolBarButton(Commands.BeginSculptMeshTool);
-		ToolbarBuilder.AddToolBarButton(Commands.BeginRemeshSculptMeshTool);
-		ToolbarBuilder.AddSeparator();
-		ToolbarBuilder.AddToolBarButton(Commands.BeginPlaneCutTool);
-		ToolbarBuilder.AddToolBarButton(Commands.BeginMirrorTool);
-		ToolbarBuilder.AddToolBarButton(Commands.BeginPolygonCutTool);
-		ToolbarBuilder.AddToolBarButton(Commands.BeginSmoothMeshTool);
-		ToolbarBuilder.AddToolBarButton(Commands.BeginOffsetMeshTool);
-		ToolbarBuilder.AddToolBarButton(Commands.BeginDisplaceMeshTool);
-		ToolbarBuilder.AddToolBarButton(Commands.BeginMeshSpaceDeformerTool);
-		ToolbarBuilder.AddToolBarButton(Commands.BeginLatticeDeformerTool);
-		ToolbarBuilder.AddSeparator();
-		ToolbarBuilder.AddToolBarButton(Commands.BeginMeshAttributePaintTool);
-		ToolbarBuilder.AddSeparator();
-		ToolbarBuilder.AddToolBarButton(Commands.BeginRemeshMeshTool);
-		ToolbarBuilder.AddToolBarButton(Commands.BeginMeshInspectorTool);
-	}
-	else if (PaletteIndex == TrianglesTabName)
-	{
-		ToolbarBuilder.AddToolBarButton(Commands.BeginTransformMeshesTool);
-		ToolbarBuilder.AddToolBarButton(Commands.BeginMeshSelectionTool);
-		ToolbarBuilder.AddSeparator();
-		ToolbarBuilder.AddToolBarButton(Commands.BeginTriEditTool);
-		ToolbarBuilder.AddToolBarButton(Commands.BeginSimplifyMeshTool);
-		ToolbarBuilder.AddToolBarButton(Commands.BeginRemeshMeshTool);
-		ToolbarBuilder.AddToolBarButton(Commands.BeginProjectToTargetTool);
-		ToolbarBuilder.AddSeparator();
-		ToolbarBuilder.AddToolBarButton(Commands.BeginEditMeshMaterialsTool);
-		ToolbarBuilder.AddToolBarButton(Commands.BeginWeldEdgesTool);
-		ToolbarBuilder.AddToolBarButton(Commands.BeginRemoveOccludedTrianglesTool);
-		ToolbarBuilder.AddToolBarButton(Commands.BeginHoleFillTool);
-		ToolbarBuilder.AddSeparator();
-		ToolbarBuilder.AddToolBarButton(Commands.BeginMeshInspectorTool);
-	}
-	else if (PaletteIndex == PolyGroupsTabName)
-	{
-		ToolbarBuilder.AddToolBarButton(Commands.BeginTransformMeshesTool);
-		ToolbarBuilder.AddToolBarButton(Commands.BeginMeshSelectionTool);
-		ToolbarBuilder.AddToolBarButton(Commands.BeginPolyGroupsTool);
-		ToolbarBuilder.AddSeparator();
-		ToolbarBuilder.AddToolBarButton(Commands.BeginPolyEditTool);
-		ToolbarBuilder.AddToolBarButton(Commands.BeginPolyDeformTool);
-		ToolbarBuilder.AddSeparator();
-		ToolbarBuilder.AddToolBarButton(Commands.BeginGroupEdgeInsertionTool);
-		ToolbarBuilder.AddToolBarButton(Commands.BeginEdgeLoopInsertionTool);
-		ToolbarBuilder.AddSeparator();
-		ToolbarBuilder.AddToolBarButton(Commands.BeginAttributeEditorTool);
-		ToolbarBuilder.AddToolBarButton(Commands.BeginMeshInspectorTool);
-	}
-	else if (PaletteIndex == UVNormalTabName)
-	{
-		ToolbarBuilder.AddToolBarButton(Commands.BeginTransformMeshesTool);
-		ToolbarBuilder.AddToolBarButton(Commands.BeginMeshSelectionTool);
-		ToolbarBuilder.AddToolBarButton(Commands.BeginPolyGroupsTool);
-		ToolbarBuilder.AddSeparator();
-		ToolbarBuilder.AddToolBarButton(Commands.BeginEditNormalsTool);
-		ToolbarBuilder.AddToolBarButton(Commands.BeginEditTangentsTool);
-		ToolbarBuilder.AddSeparator();
-		ToolbarBuilder.AddToolBarButton(Commands.BeginGlobalUVGenerateTool);
-		ToolbarBuilder.AddToolBarButton(Commands.BeginGroupUVGenerateTool);
-		ToolbarBuilder.AddToolBarButton(Commands.BeginUVProjectionTool);
-		ToolbarBuilder.AddToolBarButton(Commands.BeginUVSeamEditTool);
-		ToolbarBuilder.AddToolBarButton(Commands.BeginTransformUVIslandsTool);
-		ToolbarBuilder.AddToolBarButton(Commands.BeginUVLayoutTool);
-		ToolbarBuilder.AddSeparator();
-		ToolbarBuilder.AddToolBarButton(Commands.BeginBakeMeshAttributeMapsTool);
-		ToolbarBuilder.AddSeparator();
-		ToolbarBuilder.AddToolBarButton(Commands.BeginAttributeEditorTool);
-		ToolbarBuilder.AddToolBarButton(Commands.BeginMeshInspectorTool);
-	}
-	else if (PaletteIndex == VolumesTabName)
-	{
-		ToolbarBuilder.AddToolBarButton(Commands.BeginVolumeToMeshTool);
-		ToolbarBuilder.AddToolBarButton(Commands.BeginMeshToVolumeTool);
-		ToolbarBuilder.AddSeparator();
-		ToolbarBuilder.AddToolBarButton(Commands.BeginBspConversionTool);
-		ToolbarBuilder.AddSeparator();
-		ToolbarBuilder.AddToolBarButton(Commands.BeginPhysicsInspectorTool);
-		ToolbarBuilder.AddToolBarButton(Commands.BeginSetCollisionGeometryTool);
-		//ToolbarBuilder.AddToolBarButton(Commands.BeginEditCollisionGeometryTool);
-		ToolbarBuilder.AddToolBarButton(Commands.BeginExtractCollisionGeometryTool);
-	}
-	else if (PaletteIndex == HairTabName)
-	{
-	ToolbarBuilder.AddToolBarButton(Commands.BeginTransformMeshesTool);
-	ToolbarBuilder.AddToolBarButton(Commands.BeginMeshSelectionTool);
-	ToolbarBuilder.AddSeparator();
-	ToolbarBuilder.AddToolBarButton(Commands.BeginGroomToMeshTool);
-	ToolbarBuilder.AddToolBarButton(Commands.BeginGroomCardsEditorTool);
-	ToolbarBuilder.AddToolBarButton(Commands.BeginGenerateLODMeshesTool);
-	ToolbarBuilder.AddSeparator();
-	ToolbarBuilder.AddToolBarButton(Commands.BeginAttributeEditorTool);
-	ToolbarBuilder.AddToolBarButton(Commands.BeginMeshInspectorTool);
-	}
+	check(false);
 }
 
 
@@ -496,12 +347,6 @@ void FModelingToolsEditorModeToolkit::BuildToolPalette_Experimental(FName Palett
 {
 	const FModelingToolsManagerCommands& Commands = FModelingToolsManagerCommands::Get();
 
-	ToolbarBuilder.AddToolBarButton(Commands.AcceptActiveTool);
-	ToolbarBuilder.AddToolBarButton(Commands.CancelActiveTool);
-	ToolbarBuilder.AddToolBarButton(Commands.CompleteActiveTool);
-
-	ToolbarBuilder.AddSeparator();
-
 	if (PaletteIndex == PrimitiveTabName)
 	{
 		ToolbarBuilder.AddToolBarButton(Commands.BeginAddBoxPrimitiveTool);
@@ -656,6 +501,34 @@ void FModelingToolsEditorModeToolkit::BuildToolPalette_Experimental(FName Palett
 		ToolbarBuilder.AddToolBarButton(Commands.BeginShapeSprayTool);
 		ToolbarBuilder.AddToolBarButton(Commands.BeginSubdividePolyTool);
 	}
+	else if (PaletteIndex == PolyModelTabName)
+	{
+		ToolbarBuilder.AddToolBarButton(Commands.BeginPolyModelTool_FaceSelect);
+		ToolbarBuilder.AddToolBarButton(Commands.BeginPolyModelTool_EdgeSelect);
+		ToolbarBuilder.AddToolBarButton(Commands.BeginPolyModelTool_VertexSelect);
+		ToolbarBuilder.AddToolBarButton(Commands.BeginPolyModelTool_AllSelect);
+		ToolbarBuilder.AddToolBarButton(Commands.BeginPolyModelTool_LoopSelect);
+		ToolbarBuilder.AddToolBarButton(Commands.BeginPolyModelTool_RingSelect);
+		ToolbarBuilder.AddSeparator();
+		ToolbarBuilder.AddToolBarButton(Commands.BeginPolyModelTool_Extrude);
+		ToolbarBuilder.AddToolBarButton(Commands.BeginPolyModelTool_Offset);
+		ToolbarBuilder.AddToolBarButton(Commands.BeginPolyModelTool_Inset);
+		ToolbarBuilder.AddToolBarButton(Commands.BeginPolyModelTool_Outset);
+		ToolbarBuilder.AddToolBarButton(Commands.BeginPolyModelTool_CutFaces);
+		ToolbarBuilder.AddSeparator();
+		ToolbarBuilder.AddToolBarButton(Commands.BeginGroupEdgeInsertionTool);
+		ToolbarBuilder.AddToolBarButton(Commands.BeginEdgeLoopInsertionTool);
+		ToolbarBuilder.AddSeparator();
+		ToolbarBuilder.AddToolBarButton(Commands.BeginSubdividePolyTool);
+		ToolbarBuilder.AddToolBarButton(Commands.BeginPolyEditTool);
+	}
+
+	ToolbarBuilder.AddSeparator();
+
+	ToolbarBuilder.AddToolBarButton(Commands.AcceptActiveTool);
+	ToolbarBuilder.AddToolBarButton(Commands.CancelActiveTool);
+	ToolbarBuilder.AddToolBarButton(Commands.CompleteActiveTool);
+
 }
 
 void FModelingToolsEditorModeToolkit::OnToolPaletteChanged(FName PaletteName) 
diff --git a/Engine/Plugins/Experimental/ModelingToolsEditorMode/Source/ModelingToolsEditorMode/Private/ModelingToolsManagerActions.cpp b/Engine/Plugins/Experimental/ModelingToolsEditorMode/Source/ModelingToolsEditorMode/Private/ModelingToolsManagerActions.cpp
--- a/Engine/Plugins/Experimental/ModelingToolsEditorMode/Source/ModelingToolsEditorMode/Private/ModelingToolsManagerActions.cpp
+++ b/Engine/Plugins/Experimental/ModelingToolsEditorMode/Source/ModelingToolsEditorMode/Private/ModelingToolsManagerActions.cpp
@@ -108,6 +108,21 @@ void FModelingToolsManagerCommands::RegisterCommands()
 	UI_COMMAND(BeginGroomCardsEditorTool, "CardsEditor", "Start the Cards Editor Tool", EUserInterfaceActionType::ToggleButton, FInputChord());
 	UI_COMMAND(BeginSubdividePolyTool, "Subdivide", "Start the SubdividePoly Tool", EUserInterfaceActionType::ToggleButton, FInputChord());
 
+
+	UI_COMMAND(BeginPolyModelTool_FaceSelect, "Faces", "Start the PolyGroup Face Selection Tool", EUserInterfaceActionType::ToggleButton, FInputChord());
+	UI_COMMAND(BeginPolyModelTool_EdgeSelect, "Edges", "Start the PolyGroup Edge Selection Tool", EUserInterfaceActionType::ToggleButton, FInputChord());
+	UI_COMMAND(BeginPolyModelTool_VertexSelect, "Verts", "Start the PolyGroup Vertex Selection Tool", EUserInterfaceActionType::ToggleButton, FInputChord());
+	UI_COMMAND(BeginPolyModelTool_AllSelect, "Any", "Start the PolyGroup Face/Edge/Vertex Selection Tool", EUserInterfaceActionType::ToggleButton, FInputChord());
+	UI_COMMAND(BeginPolyModelTool_LoopSelect, "Loops", "Start the PolyGroup Loop Selection Tool", EUserInterfaceActionType::ToggleButton, FInputChord());
+	UI_COMMAND(BeginPolyModelTool_RingSelect, "Rings", "Start the PolyGroup Ring Selection Tool", EUserInterfaceActionType::ToggleButton, FInputChord());
+
+	UI_COMMAND(BeginPolyModelTool_Extrude, "Extrude", "Start the PolyGroup Extrude Tool", EUserInterfaceActionType::ToggleButton, FInputChord());
+	UI_COMMAND(BeginPolyModelTool_Offset, "Offset", "Start the PolyGroup Offset Tool", EUserInterfaceActionType::ToggleButton, FInputChord());
+	UI_COMMAND(BeginPolyModelTool_Inset, "Inset", "Start the PolyGroup Inset Tool", EUserInterfaceActionType::ToggleButton, FInputChord());
+	UI_COMMAND(BeginPolyModelTool_Outset, "Outset", "Start the PolyGroup Outset Tool", EUserInterfaceActionType::ToggleButton, FInputChord());
+	UI_COMMAND(BeginPolyModelTool_CutFaces, "Cut", "Start the PolyGroup Cut Faces Tool", EUserInterfaceActionType::ToggleButton, FInputChord());
+
+
 	UI_COMMAND(AcceptActiveTool, "Accept", "Accept the active tool", EUserInterfaceActionType::Button, FInputChord());
 	UI_COMMAND(CancelActiveTool, "Cancel", "Cancel the active tool", EUserInterfaceActionType::Button, FInputChord());
 	UI_COMMAND(CompleteActiveTool, "Complete", "Complete the active tool", EUserInterfaceActionType::Button, FInputChord());
diff --git a/Engine/Plugins/Experimental/ModelingToolsEditorMode/Source/ModelingToolsEditorMode/Public/ModelingToolsManagerActions.h b/Engine/Plugins/Experimental/ModelingToolsEditorMode/Source/ModelingToolsEditorMode/Public/ModelingToolsManagerActions.h
--- a/Engine/Plugins/Experimental/ModelingToolsEditorMode/Source/ModelingToolsEditorMode/Public/ModelingToolsManagerActions.h
+++ b/Engine/Plugins/Experimental/ModelingToolsEditorMode/Source/ModelingToolsEditorMode/Public/ModelingToolsManagerActions.h
@@ -101,6 +101,21 @@ public:
 	TSharedPtr<FUICommandInfo> BeginGroomMeshToCardsTool;
 	TSharedPtr<FUICommandInfo> BeginGroomCardsEditorTool;
 
+
+	TSharedPtr<FUICommandInfo> BeginPolyModelTool_FaceSelect;
+	TSharedPtr<FUICommandInfo> BeginPolyModelTool_EdgeSelect;
+	TSharedPtr<FUICommandInfo> BeginPolyModelTool_VertexSelect;
+	TSharedPtr<FUICommandInfo> BeginPolyModelTool_AllSelect;
+	TSharedPtr<FUICommandInfo> BeginPolyModelTool_LoopSelect;
+	TSharedPtr<FUICommandInfo> BeginPolyModelTool_RingSelect;
+
+	TSharedPtr<FUICommandInfo> BeginPolyModelTool_Extrude;
+	TSharedPtr<FUICommandInfo> BeginPolyModelTool_Offset;
+	TSharedPtr<FUICommandInfo> BeginPolyModelTool_Inset;
+	TSharedPtr<FUICommandInfo> BeginPolyModelTool_Outset;
+	TSharedPtr<FUICommandInfo> BeginPolyModelTool_CutFaces;
+
+
 	TSharedPtr<FUICommandInfo> AcceptActiveTool;
 	TSharedPtr<FUICommandInfo> CancelActiveTool;
 	TSharedPtr<FUICommandInfo> CompleteActiveTool;
