Updated Modeling tools to UEdMode.

diff --git a/Engine/Plugins/Experimental/ModelingToolsEditorMode/Source/ModelingToolsEditorMode/Private/ModelingToolsEditorMode.cpp b/Engine/Plugins/Experimental/ModelingToolsEditorMode/Source/ModelingToolsEditorMode/Private/ModelingToolsEditorMode.cpp
--- a/Engine/Plugins/Experimental/ModelingToolsEditorMode/Source/ModelingToolsEditorMode/Private/ModelingToolsEditorMode.cpp
+++ b/Engine/Plugins/Experimental/ModelingToolsEditorMode/Source/ModelingToolsEditorMode/Private/ModelingToolsEditorMode.cpp
@@ -95,34 +95,33 @@
 #include "ModelingToolsManagerActions.h"
 #include "ModelingModeAssetAPI.h"
 
-#define LOCTEXT_NAMESPACE "FModelingToolsEditorMode"
+#define LOCTEXT_NAMESPACE "UModelingToolsEditorMode"
 
 
 //#define ENABLE_DEBUG_PRINTING
 
-const FEditorModeID FModelingToolsEditorMode::EM_ModelingToolsEditorModeId = TEXT("EM_ModelingToolsEditorMode");
+const FEditorModeID UModelingToolsEditorMode::EM_ModelingToolsEditorModeId = TEXT("EM_ModelingToolsEditorMode");
 
-FModelingToolsEditorMode::FModelingToolsEditorMode()
+UModelingToolsEditorMode::UModelingToolsEditorMode()
 {
-	ToolsContext = nullptr;
-
-	UICommandList = MakeShareable(new FUICommandList);
+	Info = FEditorModeInfo(
+		EM_ModelingToolsEditorModeId,
+		LOCTEXT("ModelingToolsEditorModeName", "Modeling"),
+		FSlateIcon("ModelingToolsStyle", "LevelEditor.ModelingToolsMode", "LevelEditor.ModelingToolsMode.Small"),
+		true);
 }
 
-FModelingToolsEditorMode::~FModelingToolsEditorMode()
+UModelingToolsEditorMode::UModelingToolsEditorMode(FVTableHelper& Helper)
 {
-	ToolsContext = nullptr;
 }
 
-
-void FModelingToolsEditorMode::ActorSelectionChangeNotify()
+UModelingToolsEditorMode::~UModelingToolsEditorMode()
 {
 }
 
-
-bool FModelingToolsEditorMode::ProcessEditDelete()
+bool UModelingToolsEditorMode::ProcessEditDelete()
 {
-	if (ToolsContext->ProcessEditDelete())
+	if (UEdMode::ProcessEditDelete())
 	{
 		return true;
 	}
@@ -146,7 +145,7 @@ bool FModelingToolsEditorMode::ProcessEditDelete()
 }
 
 
-bool FModelingToolsEditorMode::ProcessEditCut()
+bool UModelingToolsEditorMode::ProcessEditCut()
 {
 	// for now we disable deleting in an Accept-style tool because it can result in crashes if we are deleting target object
 	if (GetToolManager()->HasAnyActiveTool() && GetToolManager()->GetActiveTool(EToolSide::Mouse)->HasAccept())
@@ -167,36 +166,26 @@ bool FModelingToolsEditorMode::ProcessEditCut()
 
 
 
-bool FModelingToolsEditorMode::CanAutoSave() const
+bool UModelingToolsEditorMode::CanAutoSave() const
 {
 	// prevent autosave if any tool is active
 	return ToolsContext->ToolManager->HasAnyActiveTool() == false;
 }
 
-bool FModelingToolsEditorMode::ShouldDrawWidget() const
+bool UModelingToolsEditorMode::ShouldDrawWidget() const
 { 
 	// allow standard xform gizmo if we don't have an active tool
 	if (ToolsContext != nullptr && ToolsContext->ToolManager->HasAnyActiveTool())
 	{
 		return false;
 	}
-	return FEdMode::ShouldDrawWidget(); 
-}
 
-bool FModelingToolsEditorMode::UsesTransformWidget() const
-{ 
-	return true; 
+	return UBaseLegacyWidgetEdMode::ShouldDrawWidget(); 
 }
 
-
-void FModelingToolsEditorMode::Tick(FEditorViewportClient* ViewportClient, float DeltaTime)
+void UModelingToolsEditorMode::Tick(FEditorViewportClient* ViewportClient, float DeltaTime)
 {
-	FEdMode::Tick(ViewportClient, DeltaTime);
-
-	if (ToolsContext != nullptr)
-	{
-		ToolsContext->Tick(ViewportClient, DeltaTime);
-	}
+	UEdMode::Tick(ViewportClient, DeltaTime);
 
 	if (Toolkit.IsValid())
 	{
@@ -205,148 +194,6 @@ void FModelingToolsEditorMode::Tick(FEditorViewportClient* ViewportClient, float
 	}
 }
 
-
-
-
-void FModelingToolsEditorMode::Render(const FSceneView* View, FViewport* Viewport, FPrimitiveDrawInterface* PDI)
-{
-	FEdMode::Render(View, Viewport, PDI);
-
-	// we do not use PDI hit testing in modeling tools, so skip these render passes
-	if (PDI->IsHitTesting())
-	{
-		return;
-	}
-
-	if (ToolsContext != nullptr)
-	{
-		ToolsContext->Render(View, Viewport, PDI);
-	}
-}
-
-void FModelingToolsEditorMode::DrawHUD(FEditorViewportClient* ViewportClient,FViewport* Viewport,const FSceneView* View,FCanvas* Canvas)
-{
-	FEdMode::DrawHUD(ViewportClient, Viewport, View, Canvas);
-	if (ToolsContext != nullptr)
-	{
-		ToolsContext->DrawHUD(ViewportClient, Viewport, View, Canvas);
-	}
-}
-
-
-
-
-
-bool FModelingToolsEditorMode::InputKey(FEditorViewportClient* ViewportClient, FViewport* Viewport, FKey Key, EInputEvent Event)
-{
-	// try hotkeys
-	if (Event != IE_Released)
-	{
-		if (ToolsContext->ShouldIgnoreHotkeys() == false)		// allow the context to capture keyboard input if necessary
-		{
-			if (UICommandList->ProcessCommandBindings(Key, FSlateApplication::Get().GetModifierKeys(), false/*Event == IE_Repeat*/))
-			{
-				return true;
-			}
-		}
-	}
-
-	bool bHandled = ToolsContext->InputKey(ViewportClient, Viewport, Key, Event);
-	if (bHandled == false)
-	{
-		bHandled = FEdMode::InputKey(ViewportClient, Viewport, Key, Event);
-	}
-	return bHandled;
-
-	//bool bHandled = FEdMode::InputKey(ViewportClient, Viewport, Key, Event);
-	//bHandled |= ToolsContext->InputKey(ViewportClient, Viewport, Key, Event);
-	//return bHandled;
-}
-
-
-bool FModelingToolsEditorMode::InputAxis(FEditorViewportClient* InViewportClient, FViewport* Viewport, int32 ControllerId, FKey Key, float Delta, float DeltaTime)
-{
-	// mouse axes: EKeys::MouseX, EKeys::MouseY, EKeys::MouseWheelAxis
-	return FEdMode::InputAxis(InViewportClient, Viewport, ControllerId, Key, Delta, DeltaTime);
-}
-
-
-
-
-
-bool FModelingToolsEditorMode::StartTracking(FEditorViewportClient* InViewportClient, FViewport* InViewport) 
-{ 
-	bool bHandled = FEdMode::StartTracking(InViewportClient, InViewport);
-#ifdef ENABLE_DEBUG_PRINTING
-	UE_LOG(LogTemp, Warning, TEXT("START TRACKING - base handled was %d"), (int)bHandled);
-#endif
-
-	bHandled |= ToolsContext->StartTracking(InViewportClient, InViewport);
-
-	return bHandled;
-}
-
-bool FModelingToolsEditorMode::CapturedMouseMove(FEditorViewportClient* InViewportClient, FViewport* InViewport, int32 InMouseX, int32 InMouseY)
-{
-	bool bHandled = ToolsContext->CapturedMouseMove(InViewportClient, InViewport, InMouseX, InMouseY);
-	return bHandled;
-}
-
-
-bool FModelingToolsEditorMode::EndTracking(FEditorViewportClient* InViewportClient, FViewport* InViewport)
-{
-	bool bHandled = ToolsContext->EndTracking(InViewportClient, InViewport);
-	return bHandled;
-}
-
-
-
-
-
-
-
-
-bool FModelingToolsEditorMode::ReceivedFocus(FEditorViewportClient* ViewportClient, FViewport* Viewport)
-{
-#ifdef ENABLE_DEBUG_PRINTING
-	UE_LOG(LogTemp, Warning, TEXT("RECEIVED FOCUS"));
-#endif
-
-	return false;
-}
-
-bool FModelingToolsEditorMode::LostFocus(FEditorViewportClient* ViewportClient, FViewport* Viewport)
-{
-#ifdef ENABLE_DEBUG_PRINTING
-	UE_LOG(LogTemp, Warning, TEXT("LOST FOCUS"));
-#endif
-
-	return false;
-}
-
-
-
-bool FModelingToolsEditorMode::MouseEnter(FEditorViewportClient* ViewportClient, FViewport* Viewport, int32 x, int32 y)
-{
-	bool bHandled = ToolsContext->MouseEnter(ViewportClient, Viewport, x, y);
-	return bHandled;
-}
-
-bool FModelingToolsEditorMode::MouseMove(FEditorViewportClient* ViewportClient, FViewport* Viewport, int32 x, int32 y)
-{
-	bool bHandled = ToolsContext->MouseMove(ViewportClient, Viewport, x, y);
-	return bHandled;
-}
-
-bool FModelingToolsEditorMode::MouseLeave(FEditorViewportClient* ViewportClient, FViewport* Viewport)
-{
-	bool bHandled = ToolsContext->MouseLeave(ViewportClient, Viewport);
-	return bHandled;
-}
-
-
-
-
 //
 // FStylusStateTracker registers itself as a listener for stylus events and implements
 // the IToolStylusStateProviderAPI interface, which allows MeshSurfacePointTool implementations
@@ -429,83 +276,16 @@ public:
 
 };
 
-
-
-
-
-void FModelingToolsEditorMode::Enter()
+void UModelingToolsEditorMode::Enter()
 {
-	FEdMode::Enter();
+	UEdMode::Enter();
 
-	// initialize FEdMode ToolsContext adapter
-	ToolsContext = Owner->GetInteractiveToolsContext();
 	ModelingModeAssetGenerationAPI = MakeShareable(new FModelingModeAssetAPI(ToolsContext->GetAssetAPI()));
 
 	// register stylus event handler
 	StylusStateTracker = MakeUnique<FStylusStateTracker>();
 
-	if (!Toolkit.IsValid() && UsesToolkits())
-	{
-		Toolkit = MakeShareable(new FModelingToolsEditorModeToolkit);
-		Toolkit->Init(Owner->GetToolkitHost());
-
-
-		const FModelingToolsManagerCommands& ToolManagerCommands = FModelingToolsManagerCommands::Get();
-		const TSharedRef<FUICommandList>& CommandList = Toolkit->GetToolkitCommands();
-
-		CommandList->MapAction(
-			ToolManagerCommands.AcceptActiveTool,
-			FExecuteAction::CreateLambda([this]() { ToolsContext->EndTool(EToolShutdownType::Accept); }),
-			FCanExecuteAction::CreateLambda([this]() { return ToolsContext->CanAcceptActiveTool(); }),
-			FGetActionCheckState(),
-			FIsActionButtonVisible::CreateLambda([this]() {return ToolsContext->ActiveToolHasAccept(); }),
-			EUIActionRepeatMode::RepeatDisabled
-		);
-
-		CommandList->MapAction(
-			ToolManagerCommands.CancelActiveTool,
-			FExecuteAction::CreateLambda([this]() { ToolsContext->EndTool(EToolShutdownType::Cancel); }),
-			FCanExecuteAction::CreateLambda([this]() { return ToolsContext->CanCancelActiveTool(); }),
-			FGetActionCheckState(),
-			FIsActionButtonVisible::CreateLambda([this]() {return ToolsContext->ActiveToolHasAccept(); }),
-			EUIActionRepeatMode::RepeatDisabled
-		);
-
-		CommandList->MapAction(
-			ToolManagerCommands.CompleteActiveTool,
-			FExecuteAction::CreateLambda([this]() { ToolsContext->EndTool(EToolShutdownType::Completed); }),
-			FCanExecuteAction::CreateLambda([this]() { return ToolsContext->CanCompleteActiveTool(); }),
-			FGetActionCheckState(),
-			FIsActionButtonVisible::CreateLambda([this]() {return ToolsContext->CanCompleteActiveTool(); }),
-			EUIActionRepeatMode::RepeatDisabled
-		);
-		CommandList->MapAction(
-		    ToolManagerCommands.CancelOrCompleteActiveTool,
-		    FExecuteAction::CreateLambda([this]() {
-				const EToolShutdownType ShutdownType = ToolsContext->CanCancelActiveTool() ? EToolShutdownType::Cancel : EToolShutdownType::Completed;
-				ToolsContext->EndTool(ShutdownType);
-			}),
-		    FCanExecuteAction::CreateLambda([this]() {
-			    return ToolsContext->CanCompleteActiveTool() || ToolsContext->CanCancelActiveTool();
-		    }),
-		    FGetActionCheckState(),
-		    FIsActionButtonVisible::CreateLambda([this]() { return ToolsContext->CanCompleteActiveTool() || ToolsContext->CanCancelActiveTool();}),
-		    EUIActionRepeatMode::RepeatDisabled);
-	}
-
 	const FModelingToolsManagerCommands& ToolManagerCommands = FModelingToolsManagerCommands::Get();
-	const TSharedRef<FUICommandList>& CommandList = Toolkit->GetToolkitCommands();
-	auto RegisterToolFunc = [this, &CommandList](TSharedPtr<FUICommandInfo> UICommand, FString ToolIdentifier, UInteractiveToolBuilder* Builder)
-	{
-		ToolsContext->ToolManager->RegisterToolType(ToolIdentifier, Builder);
-		CommandList->MapAction( UICommand,
-			FExecuteAction::CreateLambda([this, ToolIdentifier]() { ToolsContext->StartTool(ToolIdentifier); }),
-			FCanExecuteAction::CreateLambda([this, ToolIdentifier]() { return ToolsContext->CanStartTool(ToolIdentifier); }),
-			FIsActionChecked::CreateLambda([this, ToolIdentifier]() { return ToolsContext->GetActiveToolName() == ToolIdentifier; }));
-
-		RegisteredTools.Emplace(UICommand, ToolIdentifier);
-	};
-
 
 	// register tool set
 
@@ -513,14 +293,14 @@ void FModelingToolsEditorMode::Enter()
 	// primitive tools
 	//
 	auto RegisterPrimitiveToolFunc  =
-		[this, &RegisterToolFunc](TSharedPtr<FUICommandInfo> UICommand,
+		[this](TSharedPtr<FUICommandInfo> UICommand,
 								  FString&& ToolIdentifier,
 								  UAddPrimitiveToolBuilder::EMakeMeshShapeType ShapeTypeIn)
 	{
 		auto AddPrimitiveToolBuilder = NewObject<UAddPrimitiveToolBuilder>();
 		AddPrimitiveToolBuilder->AssetAPI = ModelingModeAssetGenerationAPI.Get();
 		AddPrimitiveToolBuilder->ShapeType = ShapeTypeIn;
-		RegisterToolFunc(UICommand, ToolIdentifier, AddPrimitiveToolBuilder);
+		RegisterTool(UICommand, ToolIdentifier, AddPrimitiveToolBuilder);
 	};
 	RegisterPrimitiveToolFunc(ToolManagerCommands.BeginAddBoxPrimitiveTool,
 							  TEXT("BeginAddBoxPrimitiveTool"),
@@ -561,27 +341,27 @@ void FModelingToolsEditorMode::Enter()
 	//
 	auto AddPatchToolBuilder = NewObject<UAddPatchToolBuilder>();
 	AddPatchToolBuilder->AssetAPI = ModelingModeAssetGenerationAPI.Get();
-	RegisterToolFunc(ToolManagerCommands.BeginAddPatchTool, TEXT("AddPatchTool"), AddPatchToolBuilder);
+	RegisterTool(ToolManagerCommands.BeginAddPatchTool, TEXT("AddPatchTool"), AddPatchToolBuilder);
 
 	auto RevolveBoundaryToolBuilder = NewObject<URevolveBoundaryToolBuilder>();
 	RevolveBoundaryToolBuilder->AssetAPI = ModelingModeAssetGenerationAPI.Get();
-	RegisterToolFunc(ToolManagerCommands.BeginRevolveBoundaryTool, TEXT("RevolveBoundaryTool"), RevolveBoundaryToolBuilder);
+	RegisterTool(ToolManagerCommands.BeginRevolveBoundaryTool, TEXT("RevolveBoundaryTool"), RevolveBoundaryToolBuilder);
 
 	auto DrawPolygonToolBuilder = NewObject<UDrawPolygonToolBuilder>();
 	DrawPolygonToolBuilder->AssetAPI = ModelingModeAssetGenerationAPI.Get();
-	RegisterToolFunc(ToolManagerCommands.BeginDrawPolygonTool, TEXT("DrawPolygonTool"), DrawPolygonToolBuilder);
+	RegisterTool(ToolManagerCommands.BeginDrawPolygonTool, TEXT("DrawPolygonTool"), DrawPolygonToolBuilder);
 
 	auto DrawPolyPathToolBuilder = NewObject<UDrawPolyPathToolBuilder>();
 	DrawPolyPathToolBuilder->AssetAPI = ModelingModeAssetGenerationAPI.Get();
-	RegisterToolFunc(ToolManagerCommands.BeginDrawPolyPathTool, TEXT("DrawPolyPath"), DrawPolyPathToolBuilder);
+	RegisterTool(ToolManagerCommands.BeginDrawPolyPathTool, TEXT("DrawPolyPath"), DrawPolyPathToolBuilder);
 
 	auto DrawAndRevolveToolBuilder = NewObject<UDrawAndRevolveToolBuilder>();
 	DrawAndRevolveToolBuilder->AssetAPI = ModelingModeAssetGenerationAPI.Get();
-	RegisterToolFunc(ToolManagerCommands.BeginDrawAndRevolveTool, TEXT("RevolveTool"), DrawAndRevolveToolBuilder);
+	RegisterTool(ToolManagerCommands.BeginDrawAndRevolveTool, TEXT("RevolveTool"), DrawAndRevolveToolBuilder);
 
 	auto ShapeSprayToolBuilder = NewObject<UShapeSprayToolBuilder>();
 	ShapeSprayToolBuilder->AssetAPI = ModelingModeAssetGenerationAPI.Get();
-	RegisterToolFunc(ToolManagerCommands.BeginShapeSprayTool, TEXT("ShapeSprayTool"), ShapeSprayToolBuilder);
+	RegisterTool(ToolManagerCommands.BeginShapeSprayTool, TEXT("ShapeSprayTool"), ShapeSprayToolBuilder);
 
 
 	//
@@ -590,32 +370,33 @@ void FModelingToolsEditorMode::Enter()
 
 	auto MoveVerticesToolBuilder = NewObject<UMeshVertexSculptToolBuilder>();
 	MoveVerticesToolBuilder->StylusAPI = StylusStateTracker.Get();
-	RegisterToolFunc(ToolManagerCommands.BeginSculptMeshTool, TEXT("MoveVerticesTool"), MoveVerticesToolBuilder);
+	RegisterTool(ToolManagerCommands.BeginSculptMeshTool, TEXT("MoveVerticesTool"), MoveVerticesToolBuilder);
 
-	RegisterToolFunc(ToolManagerCommands.BeginPolyEditTool, TEXT("EditMeshPolygonsTool"), NewObject<UEditMeshPolygonsToolBuilder>());
+	RegisterTool(ToolManagerCommands.BeginPolyEditTool, TEXT("EditMeshPolygonsTool"), NewObject<UEditMeshPolygonsToolBuilder>());
 	UEditMeshPolygonsToolBuilder* TriEditBuilder = NewObject<UEditMeshPolygonsToolBuilder>();
 	TriEditBuilder->bTriangleMode = true;
-	RegisterToolFunc(ToolManagerCommands.BeginTriEditTool, TEXT("EditMeshTrianglesTool"), TriEditBuilder);
-	RegisterToolFunc(ToolManagerCommands.BeginPolyDeformTool, TEXT("DeformMeshPolygonsTool"), NewObject<UDeformMeshPolygonsToolBuilder>());
-	RegisterToolFunc(ToolManagerCommands.BeginSmoothMeshTool, TEXT("SmoothMeshTool"), NewObject<USmoothMeshToolBuilder>());
-	RegisterToolFunc(ToolManagerCommands.BeginOffsetMeshTool, TEXT("OffsetMeshTool"), NewObject<UOffsetMeshToolBuilder>());
-	RegisterToolFunc(ToolManagerCommands.BeginDisplaceMeshTool, TEXT("DisplaceMeshTool"), NewObject<UDisplaceMeshToolBuilder>());
-	RegisterToolFunc(ToolManagerCommands.BeginMeshSpaceDeformerTool, TEXT("MeshSpaceDeformerTool"), NewObject<UMeshSpaceDeformerToolBuilder>());
-	RegisterToolFunc(ToolManagerCommands.BeginTransformMeshesTool, TEXT("TransformMeshesTool"), NewObject<UTransformMeshesToolBuilder>());
-	RegisterToolFunc(ToolManagerCommands.BeginEditPivotTool, TEXT("EditPivotTool"), NewObject<UEditPivotToolBuilder>());
-	RegisterToolFunc(ToolManagerCommands.BeginAlignObjectsTool, TEXT("AlignObjects"), NewObject<UAlignObjectsToolBuilder>());
-	RegisterToolFunc(ToolManagerCommands.BeginBakeTransformTool, TEXT("BakeTransformTool"), NewObject<UBakeTransformToolBuilder>());
-	RegisterToolFunc(ToolManagerCommands.BeginTransformUVIslandsTool, TEXT("EditUVIslands"), NewObject<UEditUVIslandsToolBuilder>());
-	RegisterToolFunc(ToolManagerCommands.BeginLatticeDeformerTool, TEXT("LatticeDeformerTool"), NewObject<ULatticeDeformerToolBuilder>());
+	RegisterTool(ToolManagerCommands.BeginTriEditTool, TEXT("EditMeshTrianglesTool"), TriEditBuilder);
+	RegisterTool(ToolManagerCommands.BeginPolyDeformTool, TEXT("DeformMeshPolygonsTool"), NewObject<UDeformMeshPolygonsToolBuilder>());
+	RegisterTool(ToolManagerCommands.BeginSmoothMeshTool, TEXT("SmoothMeshTool"), NewObject<USmoothMeshToolBuilder>());
+	RegisterTool(ToolManagerCommands.BeginOffsetMeshTool, TEXT("OffsetMeshTool"), NewObject<UOffsetMeshToolBuilder>());
+	RegisterTool(ToolManagerCommands.BeginDisplaceMeshTool, TEXT("DisplaceMeshTool"), NewObject<UDisplaceMeshToolBuilder>());
+	RegisterTool(ToolManagerCommands.BeginMeshSpaceDeformerTool, TEXT("MeshSpaceDeformerTool"), NewObject<UMeshSpaceDeformerToolBuilder>());
+	RegisterTool(ToolManagerCommands.BeginTransformMeshesTool, TEXT("TransformMeshesTool"), NewObject<UTransformMeshesToolBuilder>());
+	RegisterTool(ToolManagerCommands.BeginEditPivotTool, TEXT("EditPivotTool"), NewObject<UEditPivotToolBuilder>());
+	RegisterTool(ToolManagerCommands.BeginAlignObjectsTool, TEXT("AlignObjects"), NewObject<UAlignObjectsToolBuilder>());
+	RegisterTool(ToolManagerCommands.BeginBakeTransformTool, TEXT("BakeTransformTool"), NewObject<UBakeTransformToolBuilder>());
+	RegisterTool(ToolManagerCommands.BeginTransformUVIslandsTool, TEXT("EditUVIslands"), NewObject<UEditUVIslandsToolBuilder>());
+	RegisterTool(ToolManagerCommands.BeginLatticeDeformerTool, TEXT("LatticeDeformerTool"), NewObject<ULatticeDeformerToolBuilder>());
+	RegisterTool(ToolManagerCommands.BeginSubdividePolyTool, TEXT("SubdividePolyTool"), NewObject<USubdividePolyToolBuilder>());
 
 	UCombineMeshesToolBuilder* CombineMeshesToolBuilder = NewObject<UCombineMeshesToolBuilder>();
 	CombineMeshesToolBuilder->AssetAPI = ModelingModeAssetGenerationAPI.Get();
-	RegisterToolFunc(ToolManagerCommands.BeginCombineMeshesTool, TEXT("CombineMeshesTool"), CombineMeshesToolBuilder);
+	RegisterTool(ToolManagerCommands.BeginCombineMeshesTool, TEXT("CombineMeshesTool"), CombineMeshesToolBuilder);
 
 	UCombineMeshesToolBuilder* DuplicateMeshesToolBuilder = NewObject<UCombineMeshesToolBuilder>();
 	DuplicateMeshesToolBuilder->AssetAPI = ModelingModeAssetGenerationAPI.Get();
 	DuplicateMeshesToolBuilder->bIsDuplicateTool = true;
-	RegisterToolFunc(ToolManagerCommands.BeginDuplicateMeshesTool, TEXT("DuplicateMeshesTool"), DuplicateMeshesToolBuilder);
+	RegisterTool(ToolManagerCommands.BeginDuplicateMeshesTool, TEXT("DuplicateMeshesTool"), DuplicateMeshesToolBuilder);
 
 
 	// edit tools
@@ -624,138 +405,136 @@ void FModelingToolsEditorMode::Enter()
 	auto DynaSculptToolBuilder = NewObject<UDynamicMeshSculptToolBuilder>();
 	DynaSculptToolBuilder->bEnableRemeshing = true;
 	DynaSculptToolBuilder->StylusAPI = StylusStateTracker.Get();
-	RegisterToolFunc(ToolManagerCommands.BeginRemeshSculptMeshTool, TEXT("DynaSculptTool"), DynaSculptToolBuilder);
+	RegisterTool(ToolManagerCommands.BeginRemeshSculptMeshTool, TEXT("DynaSculptTool"), DynaSculptToolBuilder);
 
-	RegisterToolFunc(ToolManagerCommands.BeginRemeshMeshTool, TEXT("RemeshMeshTool"), NewObject<URemeshMeshToolBuilder>());
-	RegisterToolFunc(ToolManagerCommands.BeginProjectToTargetTool, TEXT("ProjectToTargetTool"), NewObject<UProjectToTargetToolBuilder>());
-	RegisterToolFunc(ToolManagerCommands.BeginSimplifyMeshTool, TEXT("SimplifyMeshTool"), NewObject<USimplifyMeshToolBuilder>());
+	RegisterTool(ToolManagerCommands.BeginRemeshMeshTool, TEXT("RemeshMeshTool"), NewObject<URemeshMeshToolBuilder>());
+	RegisterTool(ToolManagerCommands.BeginProjectToTargetTool, TEXT("ProjectToTargetTool"), NewObject<UProjectToTargetToolBuilder>());
+	RegisterTool(ToolManagerCommands.BeginSimplifyMeshTool, TEXT("SimplifyMeshTool"), NewObject<USimplifyMeshToolBuilder>());
 
 	auto GroupEdgeInsertionToolBuilder = NewObject<UGroupEdgeInsertionToolBuilder>();
 	GroupEdgeInsertionToolBuilder->AssetAPI = ToolsContext->GetAssetAPI();
-	RegisterToolFunc(ToolManagerCommands.BeginGroupEdgeInsertionTool, TEXT("GroupEdgeInsertionTool"), GroupEdgeInsertionToolBuilder);
+	RegisterTool(ToolManagerCommands.BeginGroupEdgeInsertionTool, TEXT("GroupEdgeInsertionTool"), GroupEdgeInsertionToolBuilder);
 
 	auto EdgeLoopInsertionToolBuilder = NewObject<UEdgeLoopInsertionToolBuilder>();
 	EdgeLoopInsertionToolBuilder->AssetAPI = ToolsContext->GetAssetAPI();
-	RegisterToolFunc(ToolManagerCommands.BeginEdgeLoopInsertionTool, TEXT("EdgeLoopInsertionTool"), EdgeLoopInsertionToolBuilder);
-
-	RegisterToolFunc(ToolManagerCommands.BeginSubdividePolyTool, TEXT("SubdividePolyTool"), NewObject<USubdividePolyToolBuilder>());
+	RegisterTool(ToolManagerCommands.BeginEdgeLoopInsertionTool, TEXT("EdgeLoopInsertionTool"), EdgeLoopInsertionToolBuilder);
 
 	auto EditNormalsToolBuilder = NewObject<UEditNormalsToolBuilder>();
 	EditNormalsToolBuilder->AssetAPI = ModelingModeAssetGenerationAPI.Get();
-	RegisterToolFunc(ToolManagerCommands.BeginEditNormalsTool, TEXT("EditNormalsTool"), EditNormalsToolBuilder);
+	RegisterTool(ToolManagerCommands.BeginEditNormalsTool, TEXT("EditNormalsTool"), EditNormalsToolBuilder);
 
 	auto TangentsToolBuilder = NewObject<UMeshTangentsToolBuilder>();
-	RegisterToolFunc(ToolManagerCommands.BeginEditTangentsTool, TEXT("MeshTangentsTool"), TangentsToolBuilder);
+	RegisterTool(ToolManagerCommands.BeginEditTangentsTool, TEXT("MeshTangentsTool"), TangentsToolBuilder);
 
 	auto RemoveOccludedTrianglesToolBuilder = NewObject<URemoveOccludedTrianglesToolBuilder>();
 	RemoveOccludedTrianglesToolBuilder->AssetAPI = ModelingModeAssetGenerationAPI.Get();
-	RegisterToolFunc(ToolManagerCommands.BeginRemoveOccludedTrianglesTool, TEXT("RemoveOccludedTrianglesTool"), RemoveOccludedTrianglesToolBuilder);
+	RegisterTool(ToolManagerCommands.BeginRemoveOccludedTrianglesTool, TEXT("RemoveOccludedTrianglesTool"), RemoveOccludedTrianglesToolBuilder);
 
 	auto HoleFillToolBuilder = NewObject<UHoleFillToolBuilder>();
-	RegisterToolFunc(ToolManagerCommands.BeginHoleFillTool, TEXT("HoleFillTool"), HoleFillToolBuilder);
+	RegisterTool(ToolManagerCommands.BeginHoleFillTool, TEXT("HoleFillTool"), HoleFillToolBuilder);
 
 	auto UVProjectionToolBuilder = NewObject<UUVProjectionToolBuilder>();
 	UVProjectionToolBuilder->AssetAPI = ModelingModeAssetGenerationAPI.Get();
-	RegisterToolFunc(ToolManagerCommands.BeginUVProjectionTool, TEXT("UVProjectionTool"), UVProjectionToolBuilder);
+	RegisterTool(ToolManagerCommands.BeginUVProjectionTool, TEXT("UVProjectionTool"), UVProjectionToolBuilder);
 
 	auto UVLayoutToolBuilder = NewObject<UUVLayoutToolBuilder>();
 	UVLayoutToolBuilder->AssetAPI = ModelingModeAssetGenerationAPI.Get();
-	RegisterToolFunc(ToolManagerCommands.BeginUVLayoutTool, TEXT("UVLayoutTool"), UVLayoutToolBuilder);
+	RegisterTool(ToolManagerCommands.BeginUVLayoutTool, TEXT("UVLayoutTool"), UVLayoutToolBuilder);
 
 	auto MergeMeshesToolBuilder = NewObject<UMergeMeshesToolBuilder>();
 	MergeMeshesToolBuilder->AssetAPI = ModelingModeAssetGenerationAPI.Get();
-	RegisterToolFunc(ToolManagerCommands.BeginVoxelMergeTool, TEXT("MergeMeshesTool"), MergeMeshesToolBuilder);
+	RegisterTool(ToolManagerCommands.BeginVoxelMergeTool, TEXT("MergeMeshesTool"), MergeMeshesToolBuilder);
 
 	auto VoxelCSGMeshesToolBuilder = NewObject<UVoxelCSGMeshesToolBuilder>();
 	VoxelCSGMeshesToolBuilder->AssetAPI = ModelingModeAssetGenerationAPI.Get();
-	RegisterToolFunc(ToolManagerCommands.BeginVoxelBooleanTool, TEXT("VoxelCSGMeshesTool"), VoxelCSGMeshesToolBuilder);
+	RegisterTool(ToolManagerCommands.BeginVoxelBooleanTool, TEXT("VoxelCSGMeshesTool"), VoxelCSGMeshesToolBuilder);
 
 	auto VoxelSolidifyMeshesToolBuilder = NewObject<UVoxelSolidifyMeshesToolBuilder>();
 	VoxelSolidifyMeshesToolBuilder->AssetAPI = ModelingModeAssetGenerationAPI.Get();
-	RegisterToolFunc(ToolManagerCommands.BeginVoxelSolidifyTool, TEXT("VoxelSolidifyMeshesTool"), VoxelSolidifyMeshesToolBuilder);
+	RegisterTool(ToolManagerCommands.BeginVoxelSolidifyTool, TEXT("VoxelSolidifyMeshesTool"), VoxelSolidifyMeshesToolBuilder);
 
 	auto VoxelBlendMeshesToolBuilder = NewObject<UVoxelBlendMeshesToolBuilder>();
 	VoxelBlendMeshesToolBuilder->AssetAPI = ModelingModeAssetGenerationAPI.Get();
-	RegisterToolFunc(ToolManagerCommands.BeginVoxelBlendTool, TEXT("VoxelBlendMeshesTool"), VoxelBlendMeshesToolBuilder);
+	RegisterTool(ToolManagerCommands.BeginVoxelBlendTool, TEXT("VoxelBlendMeshesTool"), VoxelBlendMeshesToolBuilder);
 
 	auto VoxelMorphologyMeshesToolBuilder = NewObject<UVoxelMorphologyMeshesToolBuilder>();
 	VoxelMorphologyMeshesToolBuilder->AssetAPI = ModelingModeAssetGenerationAPI.Get();
-	RegisterToolFunc(ToolManagerCommands.BeginVoxelMorphologyTool, TEXT("VoxelMorphologyMeshesTool"), VoxelMorphologyMeshesToolBuilder);
+	RegisterTool(ToolManagerCommands.BeginVoxelMorphologyTool, TEXT("VoxelMorphologyMeshesTool"), VoxelMorphologyMeshesToolBuilder);
 
 	auto SelfUnionMeshesToolBuilder = NewObject<USelfUnionMeshesToolBuilder>();
 	SelfUnionMeshesToolBuilder->AssetAPI = ModelingModeAssetGenerationAPI.Get();
-	RegisterToolFunc(ToolManagerCommands.BeginSelfUnionTool, TEXT("SelfUnionMeshesTool"), SelfUnionMeshesToolBuilder);
+	RegisterTool(ToolManagerCommands.BeginSelfUnionTool, TEXT("SelfUnionMeshesTool"), SelfUnionMeshesToolBuilder);
 
 	auto CSGMeshesToolBuilder = NewObject<UCSGMeshesToolBuilder>();
 	CSGMeshesToolBuilder->AssetAPI = ModelingModeAssetGenerationAPI.Get();
-	RegisterToolFunc(ToolManagerCommands.BeginMeshBooleanTool, TEXT("CSGMeshesTool"), CSGMeshesToolBuilder);
+	RegisterTool(ToolManagerCommands.BeginMeshBooleanTool, TEXT("CSGMeshesTool"), CSGMeshesToolBuilder);
 
 	auto BspConversionToolBuilder = NewObject<UBspConversionToolBuilder>();
 	BspConversionToolBuilder->AssetAPI = ModelingModeAssetGenerationAPI.Get();
-	RegisterToolFunc(ToolManagerCommands.BeginBspConversionTool, TEXT("BspConversionTool"), BspConversionToolBuilder);
+	RegisterTool(ToolManagerCommands.BeginBspConversionTool, TEXT("BspConversionTool"), BspConversionToolBuilder);
 
 	auto MeshToVolumeToolBuilder = NewObject<UMeshToVolumeToolBuilder>();
 	MeshToVolumeToolBuilder->AssetAPI = ModelingModeAssetGenerationAPI.Get();
-	RegisterToolFunc(ToolManagerCommands.BeginMeshToVolumeTool, TEXT("MeshToVolumeTool"), MeshToVolumeToolBuilder);
+	RegisterTool(ToolManagerCommands.BeginMeshToVolumeTool, TEXT("MeshToVolumeTool"), MeshToVolumeToolBuilder);
 
 	auto VolumeToMeshToolBuilder = NewObject<UVolumeToMeshToolBuilder>();
 	VolumeToMeshToolBuilder->AssetAPI = ModelingModeAssetGenerationAPI.Get();
-	RegisterToolFunc(ToolManagerCommands.BeginVolumeToMeshTool, TEXT("VolumeToMeshTool"), VolumeToMeshToolBuilder);
+	RegisterTool(ToolManagerCommands.BeginVolumeToMeshTool, TEXT("VolumeToMeshTool"), VolumeToMeshToolBuilder);
 
 	auto PlaneCutToolBuilder = NewObject<UPlaneCutToolBuilder>();
 	PlaneCutToolBuilder->AssetAPI = ModelingModeAssetGenerationAPI.Get();
-	RegisterToolFunc(ToolManagerCommands.BeginPlaneCutTool, TEXT("PlaneCutTool"), PlaneCutToolBuilder);
+	RegisterTool(ToolManagerCommands.BeginPlaneCutTool, TEXT("PlaneCutTool"), PlaneCutToolBuilder);
 
 	auto MirrorToolBuilder = NewObject<UMirrorToolBuilder>();
 	MirrorToolBuilder->AssetAPI = ModelingModeAssetGenerationAPI.Get();
-	RegisterToolFunc(ToolManagerCommands.BeginMirrorTool, TEXT("MirrorTool"), MirrorToolBuilder);
+	RegisterTool(ToolManagerCommands.BeginMirrorTool, TEXT("MirrorTool"), MirrorToolBuilder);
 
 	auto PolygonCutToolBuilder = NewObject<UPolygonOnMeshToolBuilder>();
-	RegisterToolFunc(ToolManagerCommands.BeginPolygonCutTool, TEXT("PolyCutTool"), PolygonCutToolBuilder);
+	RegisterTool(ToolManagerCommands.BeginPolygonCutTool, TEXT("PolyCutTool"), PolygonCutToolBuilder);
 
 	auto GlobalUVGenerateToolBuilder = NewObject<UParameterizeMeshToolBuilder>();
 	GlobalUVGenerateToolBuilder->AssetAPI = ModelingModeAssetGenerationAPI.Get();
 	GlobalUVGenerateToolBuilder->bDoAutomaticGlobalUnwrap = true;
-	RegisterToolFunc(ToolManagerCommands.BeginGlobalUVGenerateTool, TEXT("GlobalParameterizeMeshTool"), GlobalUVGenerateToolBuilder);
+	RegisterTool(ToolManagerCommands.BeginGlobalUVGenerateTool, TEXT("GlobalParameterizeMeshTool"), GlobalUVGenerateToolBuilder);
 
 	auto GroupUVGenerateToolBuilder = NewObject<UParameterizeMeshToolBuilder>();
 	GroupUVGenerateToolBuilder->AssetAPI = ModelingModeAssetGenerationAPI.Get();
 	GroupUVGenerateToolBuilder->bDoAutomaticGlobalUnwrap = false;
-	RegisterToolFunc(ToolManagerCommands.BeginGroupUVGenerateTool, TEXT("GroupParameterizeMeshTool"), GroupUVGenerateToolBuilder);
+	RegisterTool(ToolManagerCommands.BeginGroupUVGenerateTool, TEXT("GroupParameterizeMeshTool"), GroupUVGenerateToolBuilder);
 
-	RegisterToolFunc(ToolManagerCommands.BeginUVSeamEditTool, TEXT("UVSeamSculptTool"), NewObject< USeamSculptToolBuilder>());
+	RegisterTool(ToolManagerCommands.BeginUVSeamEditTool, TEXT("UVSeamSculptTool"), NewObject< USeamSculptToolBuilder>());
 
 	auto MeshSelectionToolBuilder = NewObject<UMeshSelectionToolBuilder>();
 	MeshSelectionToolBuilder->AssetAPI = ModelingModeAssetGenerationAPI.Get();
-	RegisterToolFunc(ToolManagerCommands.BeginMeshSelectionTool, TEXT("MeshSelectionTool"), MeshSelectionToolBuilder);
+	RegisterTool(ToolManagerCommands.BeginMeshSelectionTool, TEXT("MeshSelectionTool"), MeshSelectionToolBuilder);
 
 	auto EditMeshMaterialsToolBuilder = NewObject<UEditMeshMaterialsToolBuilder>();
 	EditMeshMaterialsToolBuilder->AssetAPI = ModelingModeAssetGenerationAPI.Get();
-	RegisterToolFunc(ToolManagerCommands.BeginEditMeshMaterialsTool, TEXT("EditMaterialsTool"), EditMeshMaterialsToolBuilder);
+	RegisterTool(ToolManagerCommands.BeginEditMeshMaterialsTool, TEXT("EditMaterialsTool"), EditMeshMaterialsToolBuilder);
 	
-	RegisterToolFunc(ToolManagerCommands.BeginMeshAttributePaintTool, TEXT("MeshAttributePaintTool"), NewObject<UMeshAttributePaintToolBuilder>());
+	RegisterTool(ToolManagerCommands.BeginMeshAttributePaintTool, TEXT("MeshAttributePaintTool"), NewObject<UMeshAttributePaintToolBuilder>());
 
 	auto BakeMeshAttributeMapsToolBuilder = NewObject<UBakeMeshAttributeMapsToolBuilder>();
 	BakeMeshAttributeMapsToolBuilder->AssetAPI = ModelingModeAssetGenerationAPI.Get();
-	RegisterToolFunc(ToolManagerCommands.BeginBakeMeshAttributeMapsTool, TEXT("BakeMeshMapsTool"), BakeMeshAttributeMapsToolBuilder);
+	RegisterTool(ToolManagerCommands.BeginBakeMeshAttributeMapsTool, TEXT("BakeMeshMapsTool"), BakeMeshAttributeMapsToolBuilder);
 
 	// analysis tools
 
-	RegisterToolFunc(ToolManagerCommands.BeginMeshInspectorTool, TEXT("MeshInspectorTool"), NewObject<UMeshInspectorToolBuilder>());
-	RegisterToolFunc(ToolManagerCommands.BeginWeldEdgesTool, TEXT("WeldMeshEdgesTool"), NewObject<UWeldMeshEdgesToolBuilder>());
-	RegisterToolFunc(ToolManagerCommands.BeginPolyGroupsTool, TEXT("ConvertToPolygonsTool"), NewObject<UConvertToPolygonsToolBuilder>());
-	RegisterToolFunc(ToolManagerCommands.BeginAttributeEditorTool, TEXT("AttributeEditorTool"), NewObject<UAttributeEditorToolBuilder>());
+	RegisterTool(ToolManagerCommands.BeginMeshInspectorTool, TEXT("MeshInspectorTool"), NewObject<UMeshInspectorToolBuilder>());
+	RegisterTool(ToolManagerCommands.BeginWeldEdgesTool, TEXT("WeldMeshEdgesTool"), NewObject<UWeldMeshEdgesToolBuilder>());
+	RegisterTool(ToolManagerCommands.BeginPolyGroupsTool, TEXT("ConvertToPolygonsTool"), NewObject<UConvertToPolygonsToolBuilder>());
+	RegisterTool(ToolManagerCommands.BeginAttributeEditorTool, TEXT("AttributeEditorTool"), NewObject<UAttributeEditorToolBuilder>());
 
 
 	// Physics Tools
 
-	RegisterToolFunc(ToolManagerCommands.BeginPhysicsInspectorTool, TEXT("PhysicsInspectorTool"), NewObject<UPhysicsInspectorToolBuilder>());
-	RegisterToolFunc(ToolManagerCommands.BeginSetCollisionGeometryTool, TEXT("SetCollisionGeoTool"), NewObject<USetCollisionGeometryToolBuilder>());
-	//RegisterToolFunc(ToolManagerCommands.BeginEditCollisionGeometryTool, TEXT("EditCollisionGeoTool"), NewObject<UEditCollisionGeometryToolBuilder>());
+	RegisterTool(ToolManagerCommands.BeginPhysicsInspectorTool, TEXT("PhysicsInspectorTool"), NewObject<UPhysicsInspectorToolBuilder>());
+	RegisterTool(ToolManagerCommands.BeginSetCollisionGeometryTool, TEXT("SetCollisionGeoTool"), NewObject<USetCollisionGeometryToolBuilder>());
+	//RegisterTool(ToolManagerCommands.BeginEditCollisionGeometryTool, TEXT("EditCollisionGeoTool"), NewObject<UEditCollisionGeometryToolBuilder>());
 
 	auto ExtractCollisionGeoToolBuilder = NewObject<UExtractCollisionGeometryToolBuilder>();
 	ExtractCollisionGeoToolBuilder->AssetAPI = ModelingModeAssetGenerationAPI.Get();
-	RegisterToolFunc(ToolManagerCommands.BeginExtractCollisionGeometryTool, TEXT("ExtractCollisionGeoTool"), ExtractCollisionGeoToolBuilder);
+	RegisterTool(ToolManagerCommands.BeginExtractCollisionGeometryTool, TEXT("ExtractCollisionGeoTool"), ExtractCollisionGeoToolBuilder);
 
 
 
@@ -763,21 +542,20 @@ void FModelingToolsEditorMode::Enter()
 
 	UGroomToMeshToolBuilder* GroomToMeshToolBuilder = NewObject<UGroomToMeshToolBuilder>();
 	GroomToMeshToolBuilder->AssetAPI = ModelingModeAssetGenerationAPI.Get();
-	RegisterToolFunc(ToolManagerCommands.BeginGroomToMeshTool, TEXT("GroomToMeshTool"), GroomToMeshToolBuilder);
+	RegisterTool(ToolManagerCommands.BeginGroomToMeshTool, TEXT("GroomToMeshTool"), GroomToMeshToolBuilder);
 
-	RegisterToolFunc(ToolManagerCommands.BeginGroomCardsEditorTool, TEXT("GroomCardsEditorTool"), NewObject<UGroomCardsEditorToolBuilder>());
+	RegisterTool(ToolManagerCommands.BeginGroomCardsEditorTool, TEXT("GroomCardsEditorTool"), NewObject<UGroomCardsEditorToolBuilder>());
 
 	UGenerateLODMeshesToolBuilder* GenerateLODMeshesToolBuilder = NewObject<UGenerateLODMeshesToolBuilder>();
 	GenerateLODMeshesToolBuilder->AssetAPI = ModelingModeAssetGenerationAPI.Get();
-	RegisterToolFunc(ToolManagerCommands.BeginGenerateLODMeshesTool, TEXT("GenerateLODMeshesTool"), GenerateLODMeshesToolBuilder);
-
+	RegisterTool(ToolManagerCommands.BeginGenerateLODMeshesTool, TEXT("GenerateLODMeshesTool"), GenerateLODMeshesToolBuilder);
 
 	// PolyModeling tools
 	auto RegisterPolyModelSelectTool = [&](EEditMeshPolygonsToolSelectionMode SelectionMode, TSharedPtr<FUICommandInfo> UICommand, FString StringName)
 	{
 		UEditMeshPolygonsSelectionModeToolBuilder* SelectionModeBuilder = NewObject<UEditMeshPolygonsSelectionModeToolBuilder>();
 		SelectionModeBuilder->SelectionMode = SelectionMode;
-		RegisterToolFunc(UICommand, StringName, SelectionModeBuilder);
+		RegisterTool(UICommand, StringName, SelectionModeBuilder);
 	};
 	RegisterPolyModelSelectTool(EEditMeshPolygonsToolSelectionMode::Faces, ToolManagerCommands.BeginPolyModelTool_FaceSelect, TEXT("PolyEdit_FaceSelect"));
 	RegisterPolyModelSelectTool(EEditMeshPolygonsToolSelectionMode::Edges, ToolManagerCommands.BeginPolyModelTool_EdgeSelect, TEXT("PolyEdit_EdgeSelect"));
@@ -790,7 +568,7 @@ void FModelingToolsEditorMode::Enter()
 	{
 		UEditMeshPolygonsActionModeToolBuilder* ActionModeBuilder = NewObject<UEditMeshPolygonsActionModeToolBuilder>();
 		ActionModeBuilder->StartupAction = Action;
-		RegisterToolFunc(UICommand, StringName, ActionModeBuilder);
+		RegisterTool(UICommand, StringName, ActionModeBuilder);
 	};
 	RegisterPolyModelActionTool(EEditMeshPolygonsToolActions::Extrude, ToolManagerCommands.BeginPolyModelTool_Extrude, TEXT("PolyEdit_Extrude"));
 	RegisterPolyModelActionTool(EEditMeshPolygonsToolActions::Offset, ToolManagerCommands.BeginPolyModelTool_Offset, TEXT("PolyEdit_Offset"));
@@ -798,23 +576,15 @@ void FModelingToolsEditorMode::Enter()
 	RegisterPolyModelActionTool(EEditMeshPolygonsToolActions::Outset, ToolManagerCommands.BeginPolyModelTool_Outset, TEXT("PolyEdit_Outset"));
 	RegisterPolyModelActionTool(EEditMeshPolygonsToolActions::CutFaces, ToolManagerCommands.BeginPolyModelTool_CutFaces, TEXT("PolyEdit_CutFaces"));
 
+
+
 	ToolsContext->ToolManager->SelectActiveToolType(EToolSide::Left, TEXT("DynaSculptTool"));
 
 	// register modeling mode hotkeys
-	FModelingModeActionCommands::RegisterCommandBindings(UICommandList, [this](EModelingModeActionCommands Command) {
+	FModelingModeActionCommands::RegisterCommandBindings(ToolCommandList, [this](EModelingModeActionCommands Command) {
 		ModelingModeShortcutRequested(Command);
 	});
 
-	// listen for Tool start/end events to bind/unbind any hotkeys relevant to that Tool
-	ToolsContext->ToolManager->OnToolStarted.AddLambda([this](UInteractiveToolManager* Manager, UInteractiveTool* Tool)
-	{
-		FModelingToolActionCommands::UpdateToolCommandBinding(Tool, UICommandList, false);
-	});
-	ToolsContext->ToolManager->OnToolEnded.AddLambda([this](UInteractiveToolManager* Manager, UInteractiveTool* Tool)
-	{
-		FModelingToolActionCommands::UpdateToolCommandBinding(Tool, UICommandList, true);
-	});
-
 	// enable realtime viewport override
 	ConfigureRealTimeViewportsOverride(true);
 
@@ -826,29 +596,9 @@ void FModelingToolsEditorMode::Enter()
 	{
 		FEngineAnalytics::GetProvider().RecordEvent( TEXT( "Editor.Usage.MeshModelingMode.Enter" ) );
 	}
-	ToolsContext->ToolManager->OnToolStarted.AddLambda([this](UInteractiveToolManager* Manager, UInteractiveTool* Tool)
-	{
-		if( FEngineAnalytics::IsAvailable() )
-		{
-			FEngineAnalytics::GetProvider().RecordEvent( TEXT( "Editor.Usage.MeshModelingMode.ToolStarted" ),
-														 TEXT( "DisplayName" ),
-														 Tool->GetToolInfo().ToolDisplayName.ToString() );
-		}
-	});
-	ToolsContext->ToolManager->OnToolEnded.AddLambda([this](UInteractiveToolManager* Manager, UInteractiveTool* Tool)
-	{
-		if( FEngineAnalytics::IsAvailable() )
-		{
-			FEngineAnalytics::GetProvider().RecordEvent( TEXT( "Editor.Usage.MeshModelingMode.ToolEnded" ),
-														 TEXT( "DisplayName" ),
-														 Tool->GetToolInfo().ToolDisplayName.ToString() );
-		}
-	});
 }
 
-
-
-void FModelingToolsEditorMode::Exit()
+void UModelingToolsEditorMode::Exit()
 {
 	//
 	// Engine Analytics
@@ -861,52 +611,90 @@ void FModelingToolsEditorMode::Exit()
 
 	StylusStateTracker = nullptr;
 
-	const TSharedRef<FUICommandList>& CommandList = Toolkit->GetToolkitCommands();
-	for (auto& RegisteredTool : RegisteredTools)
-	{
-		CommandList->UnmapAction(RegisteredTool.Key);
-		ToolsContext->ToolManager->UnregisterToolType(RegisteredTool.Value);
-	}
-	RegisteredTools.SetNum(0);
-
-	if (Toolkit.IsValid())
-	{
-		const FModelingToolsManagerCommands& ToolManagerCommands = FModelingToolsManagerCommands::Get();
-		const TSharedRef<FUICommandList>& ToolkitCommandList = Toolkit->GetToolkitCommands();
-		ToolkitCommandList->UnmapAction(ToolManagerCommands.AcceptActiveTool);
-		ToolkitCommandList->UnmapAction(ToolManagerCommands.CancelActiveTool);
-		ToolkitCommandList->UnmapAction(ToolManagerCommands.CancelOrCompleteActiveTool);
-		ToolkitCommandList->UnmapAction(ToolManagerCommands.CompleteActiveTool);
-
-		FToolkitManager::Get().CloseToolkit(Toolkit.ToSharedRef());
-		Toolkit.Reset();
-	}
-
-	ToolsContext = nullptr;
-	ModelingModeAssetGenerationAPI.Reset();
-
-	FModelingModeActionCommands::UnRegisterCommandBindings(UICommandList);
+	FModelingModeActionCommands::UnRegisterCommandBindings(ToolCommandList);
 
 	// clear realtime viewport override
 	ConfigureRealTimeViewportsOverride(false);
 
 	// Call base Exit method to ensure proper cleanup
-	FEdMode::Exit();
+	UEdMode::Exit();
 }
 
-bool FModelingToolsEditorMode::UsesToolkits() const
+void UModelingToolsEditorMode::CreateToolkit()
 {
-	return true;
+	Toolkit = MakeShareable(new FModelingToolsEditorModeToolkit);
 }
 
+void UModelingToolsEditorMode::OnToolStarted(UInteractiveToolManager* Manager, UInteractiveTool* Tool)
+{
+	FModelingToolActionCommands::UpdateToolCommandBinding(Tool, ToolCommandList, false);
+	
+	if( FEngineAnalytics::IsAvailable() )
+	{
+		FEngineAnalytics::GetProvider().RecordEvent( TEXT( "Editor.Usage.MeshModelingMode.ToolStarted" ),
+													 TEXT( "DisplayName" ),
+													 Tool->GetToolInfo().ToolDisplayName.ToString() );
+	}
+}
 
-void FModelingToolsEditorMode::AddReferencedObjects(FReferenceCollector& Collector)
+void UModelingToolsEditorMode::OnToolEnded(UInteractiveToolManager* Manager, UInteractiveTool* Tool)
 {
-	Collector.AddReferencedObject(ToolsContext);
+	FModelingToolActionCommands::UpdateToolCommandBinding(Tool, ToolCommandList, true);
+	
+	if( FEngineAnalytics::IsAvailable() )
+	{
+		FEngineAnalytics::GetProvider().RecordEvent( TEXT( "Editor.Usage.MeshModelingMode.ToolEnded" ),
+													 TEXT( "DisplayName" ),
+													 Tool->GetToolInfo().ToolDisplayName.ToString() );
+	}
 }
 
+void UModelingToolsEditorMode::BindCommands()
+{
+	const FModelingToolsManagerCommands& ToolManagerCommands = FModelingToolsManagerCommands::Get();
+	const TSharedRef<FUICommandList>& CommandList = Toolkit->GetToolkitCommands();
+
+	CommandList->MapAction(
+		ToolManagerCommands.AcceptActiveTool,
+		FExecuteAction::CreateLambda([this]() { ToolsContext->EndTool(EToolShutdownType::Accept); }),
+		FCanExecuteAction::CreateLambda([this]() { return ToolsContext->CanAcceptActiveTool(); }),
+		FGetActionCheckState(),
+		FIsActionButtonVisible::CreateLambda([this]() {return ToolsContext->ActiveToolHasAccept(); }),
+		EUIActionRepeatMode::RepeatDisabled
+		);
+
+	CommandList->MapAction(
+		ToolManagerCommands.CancelActiveTool,
+		FExecuteAction::CreateLambda([this]() { ToolsContext->EndTool(EToolShutdownType::Cancel); }),
+		FCanExecuteAction::CreateLambda([this]() { return ToolsContext->CanCancelActiveTool(); }),
+		FGetActionCheckState(),
+		FIsActionButtonVisible::CreateLambda([this]() {return ToolsContext->ActiveToolHasAccept(); }),
+		EUIActionRepeatMode::RepeatDisabled
+		);
+
+	CommandList->MapAction(
+		ToolManagerCommands.CompleteActiveTool,
+		FExecuteAction::CreateLambda([this]() { ToolsContext->EndTool(EToolShutdownType::Completed); }),
+		FCanExecuteAction::CreateLambda([this]() { return ToolsContext->CanCompleteActiveTool(); }),
+		FGetActionCheckState(),
+		FIsActionButtonVisible::CreateLambda([this]() {return ToolsContext->CanCompleteActiveTool(); }),
+		EUIActionRepeatMode::RepeatDisabled
+		);
+	CommandList->MapAction(
+		ToolManagerCommands.CancelOrCompleteActiveTool,
+		FExecuteAction::CreateLambda([this]() {
+			const EToolShutdownType ShutdownType = ToolsContext->CanCancelActiveTool() ? EToolShutdownType::Cancel : EToolShutdownType::Completed;
+			ToolsContext->EndTool(ShutdownType);
+			}),
+		FCanExecuteAction::CreateLambda([this]() {
+				return ToolsContext->CanCompleteActiveTool() || ToolsContext->CanCancelActiveTool();
+			}),
+				FGetActionCheckState(),
+				FIsActionButtonVisible::CreateLambda([this]() { return ToolsContext->CanCompleteActiveTool() || ToolsContext->CanCancelActiveTool(); }),
+				EUIActionRepeatMode::RepeatDisabled);
+}
 
-void FModelingToolsEditorMode::ModelingModeShortcutRequested(EModelingModeActionCommands Command)
+void UModelingToolsEditorMode::ModelingModeShortcutRequested(EModelingModeActionCommands Command)
 {
 	if (Command == EModelingModeActionCommands::FocusViewToCursor)
 	{
@@ -915,7 +703,7 @@ void FModelingToolsEditorMode::ModelingModeShortcutRequested(EModelingModeAction
 }
 
 
-void FModelingToolsEditorMode::FocusCameraAtCursorHotkey()
+void UModelingToolsEditorMode::FocusCameraAtCursorHotkey()
 {
 	FRay Ray = ToolsContext->GetLastWorldRay();
 
@@ -932,7 +720,7 @@ void FModelingToolsEditorMode::FocusCameraAtCursorHotkey()
 }
 
 
-bool FModelingToolsEditorMode::GetPivotForOrbit(FVector& OutPivot) const
+bool UModelingToolsEditorMode::GetPivotForOrbit(FVector& OutPivot) const
 {
 	if (GCurrentLevelEditingViewportClient)
 	{
@@ -944,7 +732,7 @@ bool FModelingToolsEditorMode::GetPivotForOrbit(FVector& OutPivot) const
 
 
 
-void FModelingToolsEditorMode::ConfigureRealTimeViewportsOverride(bool bEnable)
+void UModelingToolsEditorMode::ConfigureRealTimeViewportsOverride(bool bEnable)
 {
 	FLevelEditorModule& LevelEditorModule = FModuleManager::GetModuleChecked<FLevelEditorModule>("LevelEditor");
 	TSharedPtr<ILevelEditor> LevelEditor = LevelEditorModule.GetFirstLevelEditor();
diff --git a/Engine/Plugins/Experimental/ModelingToolsEditorMode/Source/ModelingToolsEditorMode/Private/ModelingToolsEditorModeModule.cpp b/Engine/Plugins/Experimental/ModelingToolsEditorMode/Source/ModelingToolsEditorMode/Private/ModelingToolsEditorModeModule.cpp
--- a/Engine/Plugins/Experimental/ModelingToolsEditorMode/Source/ModelingToolsEditorMode/Private/ModelingToolsEditorModeModule.cpp
+++ b/Engine/Plugins/Experimental/ModelingToolsEditorMode/Source/ModelingToolsEditorMode/Private/ModelingToolsEditorModeModule.cpp
@@ -1,7 +1,6 @@
 // Copyright Epic Games, Inc. All Rights Reserved.
 
 #include "ModelingToolsEditorModeModule.h"
-#include "ModelingToolsEditorMode.h"
 
 #include "ModelingToolsActions.h"
 #include "ModelingToolsManagerActions.h"
@@ -34,10 +33,6 @@ void FModelingToolsEditorModeModule::ShutdownModule()
 
 	// Unregister slate style overrides
 	FModelingToolsEditorModeStyle::Shutdown();
-
-	// This function may be called during shutdown to clean up your module.  For modules that support dynamic reloading,
-	// we call this function before unloading the module.
-	FEditorModeRegistry::Get().UnregisterMode(FModelingToolsEditorMode::EM_ModelingToolsEditorModeId);
 }
 
 
@@ -46,13 +41,6 @@ void FModelingToolsEditorModeModule::OnPostEngineInit()
 	// Register slate style overrides
 	FModelingToolsEditorModeStyle::Initialize();
 
-	// This code will execute after your module is loaded into memory; the exact timing is specified in the .uplugin file per-module
-	FEditorModeRegistry::Get().RegisterMode<FModelingToolsEditorMode>(
-		FModelingToolsEditorMode::EM_ModelingToolsEditorModeId,
-		LOCTEXT("ModelingToolsEditorModeName", "Modeling"),
-		FSlateIcon("ModelingToolsStyle", "LevelEditor.ModelingToolsMode", "LevelEditor.ModelingToolsMode.Small"),
-		true);
-
 	FModelingToolActionCommands::RegisterAllToolActions();
 	FModelingToolsManagerCommands::Register();
 	FModelingModeActionCommands::Register();
diff --git a/Engine/Plugins/Experimental/ModelingToolsEditorMode/Source/ModelingToolsEditorMode/Private/ModelingToolsEditorModeToolkit.cpp b/Engine/Plugins/Experimental/ModelingToolsEditorMode/Source/ModelingToolsEditorMode/Private/ModelingToolsEditorModeToolkit.cpp
--- a/Engine/Plugins/Experimental/ModelingToolsEditorMode/Source/ModelingToolsEditorMode/Private/ModelingToolsEditorModeToolkit.cpp
+++ b/Engine/Plugins/Experimental/ModelingToolsEditorMode/Source/ModelingToolsEditorMode/Private/ModelingToolsEditorModeToolkit.cpp
@@ -44,23 +44,13 @@ FModelingToolsEditorModeToolkit::~FModelingToolsEditorModeToolkit()
 {
 	UModelingToolsEditorModeSettings* Settings = GetMutableDefault<UModelingToolsEditorModeSettings>();
 	Settings->OnModified.Remove(AssetSettingsModifiedHandle);
-	GetToolsEditorMode()->GetToolsContext()->OnToolNotificationMessage.RemoveAll(this);
-	GetToolsEditorMode()->GetToolsContext()->OnToolWarningMessage.RemoveAll(this);
+	GetScriptableEditorMode()->GetInteractiveToolsContext()->OnToolNotificationMessage.RemoveAll(this);
+	GetScriptableEditorMode()->GetInteractiveToolsContext()->OnToolWarningMessage.RemoveAll(this);
 }
 
-void FModelingToolsEditorModeToolkit::Init(const TSharedPtr<IToolkitHost>& InitToolkitHost)
+void FModelingToolsEditorModeToolkit::Init(const TSharedPtr<IToolkitHost>& InitToolkitHost, TWeakObjectPtr<UEdMode> InOwningMode)
 {
-	FPropertyEditorModule& PropertyEditorModule = FModuleManager::GetModuleChecked<FPropertyEditorModule>("PropertyEditor");
-
-	FDetailsViewArgs DetailsViewArgs;
-	DetailsViewArgs.bAllowSearch = false;
-	DetailsViewArgs.NameAreaSettings = FDetailsViewArgs::HideNameArea;
-	DetailsViewArgs.bHideSelectionTip = true;
-	DetailsViewArgs.DefaultsOnlyVisibility = EEditDefaultsOnlyNodeVisibility::Automatic;
-	DetailsViewArgs.bShowOptions = false;
-	DetailsViewArgs.bAllowMultipleTopLevelObjects = true;
-
-	DetailsView = PropertyEditorModule.CreateDetailView(DetailsViewArgs);
+	FModeToolkit::Init(InitToolkitHost, InOwningMode);
 
 	ModeWarningArea = SNew(STextBlock)
 		.AutoWrapText(true)
@@ -106,7 +96,7 @@ void FModelingToolsEditorModeToolkit::Init(const TSharedPtr<IToolkitHost>& InitT
 
 			+ SVerticalBox::Slot().HAlign(HAlign_Fill).FillHeight(1.f)
 				[
-					DetailsView->AsShared()
+					ModeDetailsView->AsShared()
 				]
 
 			+ SVerticalBox::Slot().AutoHeight().HAlign(HAlign_Fill).Padding(5)
@@ -114,8 +104,6 @@ void FModelingToolsEditorModeToolkit::Init(const TSharedPtr<IToolkitHost>& InitT
 					MakeAssetConfigPanel()->AsShared()
 				]
 		];
-		
-	FModeToolkit::Init(InitToolkitHost);
 
 	ClearNotification();
 	ClearWarning();
@@ -123,8 +111,8 @@ void FModelingToolsEditorModeToolkit::Init(const TSharedPtr<IToolkitHost>& InitT
 	ActiveToolName = FText::GetEmpty();
 	ActiveToolMessage = FText::GetEmpty();
 
-	GetToolsEditorMode()->GetToolsContext()->OnToolNotificationMessage.AddSP(this, &FModelingToolsEditorModeToolkit::PostNotification);
-	GetToolsEditorMode()->GetToolsContext()->OnToolWarningMessage.AddSP(this, &FModelingToolsEditorModeToolkit::PostWarning);
+	GetScriptableEditorMode()->GetInteractiveToolsContext()->OnToolNotificationMessage.AddSP(this, &FModelingToolsEditorModeToolkit::PostNotification);
+	GetScriptableEditorMode()->GetInteractiveToolsContext()->OnToolWarningMessage.AddSP(this, &FModelingToolsEditorModeToolkit::PostWarning);
 }
 
 
@@ -234,10 +222,10 @@ TSharedPtr<SWidget> FModelingToolsEditorModeToolkit::MakeAssetConfigPanel()
 
 void FModelingToolsEditorModeToolkit::UpdateActiveToolProperties()
 {
-	UInteractiveTool* CurTool = GetToolsEditorMode()->GetToolManager()->GetActiveTool(EToolSide::Left);
+	UInteractiveTool* CurTool = GetScriptableEditorMode()->GetToolManager()->GetActiveTool(EToolSide::Left);
 	if (CurTool != nullptr)
 	{
-		DetailsView->SetObjects(CurTool->GetToolProperties(true));
+		ModeDetailsView->SetObjects(CurTool->GetToolProperties(true));
 	}
 }
 
@@ -277,21 +265,6 @@ FText FModelingToolsEditorModeToolkit::GetBaseToolkitName() const
 	return NSLOCTEXT("ModelingToolsEditorModeToolkit", "DisplayName", "ModelingToolsEditorMode Tool");
 }
 
-class FEdMode* FModelingToolsEditorModeToolkit::GetEditorMode() const
-{
-	return GLevelEditorModeTools().GetActiveMode(FModelingToolsEditorMode::EM_ModelingToolsEditorModeId);
-}
-
-FModelingToolsEditorMode * FModelingToolsEditorModeToolkit::GetToolsEditorMode() const
-{
-	return (FModelingToolsEditorMode *)GetEditorMode();
-}
-
-UEdModeInteractiveToolsContext* FModelingToolsEditorModeToolkit::GetToolsContext() const
-{
-	return GetToolsEditorMode()->GetToolsContext();
-}
-
 static const FName PrimitiveTabName(TEXT("Primitives"));
 static const FName CreateTabName(TEXT("Create"));
 static const FName EditTabName(TEXT("Edit"));
@@ -550,7 +523,7 @@ void FModelingToolsEditorModeToolkit::OnToolStarted(UInteractiveToolManager* Man
 {
 	UpdateActiveToolProperties();
 
-	UInteractiveTool* CurTool = GetToolsEditorMode()->GetToolManager()->GetActiveTool(EToolSide::Left);
+	UInteractiveTool* CurTool = GetScriptableEditorMode()->GetToolManager()->GetActiveTool(EToolSide::Left);
 	CurTool->OnPropertySetsModified.AddSP(this, &FModelingToolsEditorModeToolkit::UpdateActiveToolProperties);
 
 	ModeHeaderArea->SetVisibility(EVisibility::Collapsed);
@@ -559,13 +532,13 @@ void FModelingToolsEditorModeToolkit::OnToolStarted(UInteractiveToolManager* Man
 
 void FModelingToolsEditorModeToolkit::OnToolEnded(UInteractiveToolManager* Manager, UInteractiveTool* Tool)
 {
-	DetailsView->SetObject(nullptr);
+	ModeDetailsView->SetObject(nullptr);
 	ActiveToolName = FText::GetEmpty();
 	ModeHeaderArea->SetVisibility(EVisibility::Visible);
 	ModeHeaderArea->SetText(LOCTEXT("SelectToolLabel", "Select a Tool from the Toolbar"));
 	ClearNotification();
 	ClearWarning();
-	UInteractiveTool* CurTool = GetToolsEditorMode()->GetToolManager()->GetActiveTool(EToolSide::Left);
+	UInteractiveTool* CurTool = GetScriptableEditorMode()->GetToolManager()->GetActiveTool(EToolSide::Left);
 	if ( CurTool )
 	{
 		CurTool->OnPropertySetsModified.RemoveAll(this);
diff --git a/Engine/Plugins/Experimental/ModelingToolsEditorMode/Source/ModelingToolsEditorMode/Public/ModelingToolsEditorMode.h b/Engine/Plugins/Experimental/ModelingToolsEditorMode/Source/ModelingToolsEditorMode/Public/ModelingToolsEditorMode.h
--- a/Engine/Plugins/Experimental/ModelingToolsEditorMode/Source/ModelingToolsEditorMode/Public/ModelingToolsEditorMode.h
+++ b/Engine/Plugins/Experimental/ModelingToolsEditorMode/Source/ModelingToolsEditorMode/Public/ModelingToolsEditorMode.h
@@ -3,13 +3,14 @@
 #pragma once
 
 #include "CoreMinimal.h"
-#include "EdMode.h"
+#include "Tools/LegacyEdModeWidgetHelpers.h"
 
 #include "InputState.h"
 #include "InteractiveToolManager.h"
 #include "EdModeInteractiveToolsContext.h"
 #include "ModelingToolsActions.h"
 
+#include "ModelingToolsEditorMode.generated.h"
 
 class FEditorComponentSourceFactory;
 class FEditorToolAssetAPI;
@@ -17,37 +18,29 @@ class FUICommandList;
 class FStylusStateTracker;		// for stylus events
 class FModelingModeAssetAPI;
 
-class FModelingToolsEditorMode : public FEdMode
+UCLASS(Transient)
+class UModelingToolsEditorMode : public UBaseLegacyWidgetEdMode
 {
+	GENERATED_BODY()
 public:
 	const static FEditorModeID EM_ModelingToolsEditorModeId;
-public:
-	FModelingToolsEditorMode();
-	virtual ~FModelingToolsEditorMode();
 
+	UModelingToolsEditorMode();
+	UModelingToolsEditorMode(FVTableHelper& Helper);
+	~UModelingToolsEditorMode();
 	////////////////
-	// FEdMode interface
+	// UEdMode interface
 	////////////////
 
 	virtual void Tick(FEditorViewportClient* ViewportClient, float DeltaTime) override;
-	virtual void Render(const FSceneView* View, FViewport* Viewport, FPrimitiveDrawInterface* PDI) override;
-	virtual void DrawHUD(FEditorViewportClient* ViewportClient,FViewport* Viewport,const FSceneView* View,FCanvas* Canvas) override;
-
-	virtual bool UsesToolkits() const override;
-
-	virtual void AddReferencedObjects(FReferenceCollector& Collector) override;
 
 	virtual bool ShouldDrawWidget() const override;
-	virtual bool UsesTransformWidget() const override;
-
-	virtual void ActorSelectionChangeNotify() override;
-
 	virtual bool ProcessEditDelete() override;
 	virtual bool ProcessEditCut() override;
 
 	virtual bool CanAutoSave() const override;
 
-	virtual bool GetPivotForOrbit(FVector& OutPivot) const;
+	virtual bool GetPivotForOrbit(FVector& OutPivot) const override;
 
 	/*
 	 * focus events
@@ -59,106 +52,20 @@ public:
 	// called when we "end" this editor mode (ie switch to another tab)
 	virtual void Exit() override;
 
-	// called when viewport window is focused
-	virtual bool ReceivedFocus(FEditorViewportClient* ViewportClient, FViewport* Viewport) override;
-
-	// called when viewport window loses focus (ie when some other window is focused)
-	// *not* called when editor is backgrounded, but is called when editor is minimized
-	virtual bool LostFocus(FEditorViewportClient* ViewportClient, FViewport* Viewport) override;
-
-
-	/*
-	 * Mouse position events - These are called when no mouse button is down
-	 */
-
-	// called when mouse moves over viewport window
-	virtual bool MouseEnter(FEditorViewportClient* ViewportClient, FViewport* Viewport, int32 x, int32 y) override;
-
-	// called when mouse leaves viewport window
-	virtual bool MouseLeave(FEditorViewportClient* ViewportClient, FViewport* Viewport) override;
-
-	// called on any mouse-move event. *not* called during tracking/capturing, eg if any button is down
-	virtual bool MouseMove(FEditorViewportClient* ViewportClient, FViewport* Viewport, int32 x, int32 y) override;
-
-
-	/*
-	 * Input Button/Axis Events & Mouse Capture
-	 *
-	 *   event sequence for left mouse down/capture/up is:
-	 *     - InputKey( EKeys::LeftMouseButton, IE_PRESSED )
-	 *     - StartTracking()
-	 *     - CapturedMouseMove() and InputAxis() (repeated for each mouse-move)
-	 *     - Inputkey( EKeys::LeftMouseButton, IE_RELEASED )
-	 *     - EndTracking()
-	 *
-	 *   for doubleclick, we get one of the above sequence, then a second
-	 *   where instead of the first IE_PRESSED we get a IE_DoubleClick
-	 *
-	 *   for mouse wheel, we get following sequence
-	 *     - InputKey ( EKeys::MouseScrollUp / Down, IE_PRESSED )
-	 *     - InputAxis() for wheel move
-	 *     - InputKey ( EKeys::MouseScrollUp / Down, IE_RELEASED )
-	 *   it appears that there will only ever be *one* InputAxis() between the pressed/released sequenece
-	 *
-	 *   Note that this wheel event sequence can happen *during* a 
-	 *   middle-mouse tracking sequence on a wheel-button (nice!)
-	 */
-
-
-	// This is not just called for keyboard keys, it is also called for mouse down/up events!
-	// Eg for left-press we get Key = EKeys::LeftMouseButton and Event = IE_Pressed
-	// Return value indicates "handled'. If we return true for mouse press events then
-	// the StartTracking/CapturedMouseMove/EndTracking sequence is *not* called (but we
-	// also don't get MouseMove(), so the mouse-movement events appear to be lost?)
-	virtual bool InputKey(FEditorViewportClient* ViewportClient, FViewport* Viewport, FKey Key, EInputEvent Event) override;
-
-
-	// Called for 1D Axis movements - EKeys::MouseX, EKeys::MouseY, and EKeys::MouseWheelAxis
-	// Called even if we return true from InputKey which otherwise blocks tracking. 
-	// Return value indicates "handled" but has no effect on CapturedMouseMove events (ie we still get them)
-	virtual bool InputAxis(FEditorViewportClient* InViewportClient, FViewport* Viewport, int32 ControllerId, FKey Key, float Delta, float DeltaTime);
-
-
-	// called on mouse-down. return value indicates whether this was "handled" but
-	// does not mean we get exclusive capture events!
-	virtual bool StartTracking(FEditorViewportClient* InViewportClient, FViewport* InViewport) override;
-
-	// called during mouse-down mouse-move. Always called, return value is used to indicate whether
-	// we "handled" this mouse move (see FEditorModeTools::CapturedMouseMove) but does not pre-empt
-	// other editor modes from seeing this move event...
-	virtual bool CapturedMouseMove(FEditorViewportClient* InViewportClient, FViewport* InViewport, int32 InMouseX, int32 InMouseY) override;
-
-	// always called on mouse-up
-	virtual bool EndTracking(FEditorViewportClient* InViewportClient, FViewport* InViewport) override;
-
 	//////////////////
-	// End of FEdMode interface
+	// End of UEdMode interface
 	//////////////////
 
-public:
-	virtual UEdModeInteractiveToolsContext* GetToolsContext() const
-	{
-		return ToolsContext;
-	}
-
-	virtual UInteractiveToolManager* GetToolManager() const
-	{ 
-		return ToolsContext->ToolManager;
-	}
-
 protected:
-
-	UEdModeInteractiveToolsContext* ToolsContext;
-
-	/** Command list lives here so that the key bindings on the commands can be processed in the viewport. */
-	TSharedPtr<FUICommandList> UICommandList;
+	virtual void BindCommands() override;
+	virtual void CreateToolkit() override;
+	virtual void OnToolStarted(UInteractiveToolManager* Manager, UInteractiveTool* Tool) override;
+	virtual void OnToolEnded(UInteractiveToolManager* Manager, UInteractiveTool* Tool) override;
 
 	TSharedPtr<FModelingModeAssetAPI> ModelingModeAssetGenerationAPI;
 
 	TUniquePtr<FStylusStateTracker> StylusStateTracker;
 
-	TArray<TPair<TSharedPtr<FUICommandInfo>, FString>> RegisteredTools;
-
 	void ModelingModeShortcutRequested(EModelingModeActionCommands Command);
 	void FocusCameraAtCursorHotkey();
 
diff --git a/Engine/Plugins/Experimental/ModelingToolsEditorMode/Source/ModelingToolsEditorMode/Public/ModelingToolsEditorModeToolkit.h b/Engine/Plugins/Experimental/ModelingToolsEditorMode/Source/ModelingToolsEditorMode/Public/ModelingToolsEditorModeToolkit.h
--- a/Engine/Plugins/Experimental/ModelingToolsEditorMode/Source/ModelingToolsEditorMode/Public/ModelingToolsEditorModeToolkit.h
+++ b/Engine/Plugins/Experimental/ModelingToolsEditorMode/Source/ModelingToolsEditorMode/Public/ModelingToolsEditorModeToolkit.h
@@ -24,17 +24,13 @@ public:
 	~FModelingToolsEditorModeToolkit();
 	
 	/** FModeToolkit interface */
-	virtual void Init(const TSharedPtr<IToolkitHost>& InitToolkitHost) override;
+	virtual void Init(const TSharedPtr<IToolkitHost>& InitToolkitHost, TWeakObjectPtr<UEdMode> InOwningMode) override;
 
 	/** IToolkit interface */
 	virtual FName GetToolkitFName() const override;
 	virtual FText GetBaseToolkitName() const override;
-	virtual class FEdMode* GetEditorMode() const override;
 	virtual TSharedPtr<class SWidget> GetInlineContent() const override { return ToolkitWidget; }
 
-	virtual FModelingToolsEditorMode* GetToolsEditorMode() const;
-	virtual UEdModeInteractiveToolsContext* GetToolsContext() const;
-
 	// set/clear notification message area
 	virtual void PostNotification(const FText& Message);
 	virtual void ClearNotification();
@@ -70,7 +66,6 @@ private:
 	FText ActiveToolMessage;
 
 	TSharedPtr<SWidget> ToolkitWidget;
-	TSharedPtr<IDetailsView> DetailsView;
 	void UpdateActiveToolProperties();
 
 
