ModelingMode: expose new GenerateStaticMeshLODAssetTool in Modeling Mode

diff --git a/Engine/Plugins/Experimental/ModelingToolsEditorMode/ModelingToolsEditorMode.uplugin b/Engine/Plugins/Experimental/ModelingToolsEditorMode/ModelingToolsEditorMode.uplugin
--- a/Engine/Plugins/Experimental/ModelingToolsEditorMode/ModelingToolsEditorMode.uplugin
+++ b/Engine/Plugins/Experimental/ModelingToolsEditorMode/ModelingToolsEditorMode.uplugin
@@ -26,6 +26,10 @@
 			"Name": "MeshModelingToolset",
 			"Enabled": true
 		},
+		{
+			"Name": "MeshLODToolset",
+			"Enabled": true
+		},
 		{
 			"Name": "StylusInput",
 			"Enabled": true
diff --git a/Engine/Plugins/Experimental/ModelingToolsEditorMode/Source/ModelingToolsEditorMode/ModelingToolsEditorMode.Build.cs b/Engine/Plugins/Experimental/ModelingToolsEditorMode/Source/ModelingToolsEditorMode/ModelingToolsEditorMode.Build.cs
--- a/Engine/Plugins/Experimental/ModelingToolsEditorMode/Source/ModelingToolsEditorMode/ModelingToolsEditorMode.Build.cs
+++ b/Engine/Plugins/Experimental/ModelingToolsEditorMode/Source/ModelingToolsEditorMode/ModelingToolsEditorMode.Build.cs
@@ -55,6 +55,7 @@ public class ModelingToolsEditorMode : ModuleRules
 				"ModelingComponents",
 				"MeshModelingTools",
 				"MeshModelingToolsEditorOnly",
+				"MeshLODToolset",
 
                 "ViewportInteraction",
 
diff --git a/Engine/Plugins/Experimental/ModelingToolsEditorMode/Source/ModelingToolsEditorMode/Private/ModelingToolsEditorMode.cpp b/Engine/Plugins/Experimental/ModelingToolsEditorMode/Source/ModelingToolsEditorMode/Private/ModelingToolsEditorMode.cpp
--- a/Engine/Plugins/Experimental/ModelingToolsEditorMode/Source/ModelingToolsEditorMode/Private/ModelingToolsEditorMode.cpp
+++ b/Engine/Plugins/Experimental/ModelingToolsEditorMode/Source/ModelingToolsEditorMode/Private/ModelingToolsEditorMode.cpp
@@ -82,6 +82,9 @@
 #include "Hair/GroomCardsEditorTool.h"
 #include "GenerateLODMeshesTool.h"
 
+// asset tools
+#include "Tools/GenerateStaticMeshLODAssetTool.h"
+
 #include "EditorModeManager.h"
 
 // stylus support
@@ -398,6 +401,10 @@ void UModelingToolsEditorMode::Enter()
 	DuplicateMeshesToolBuilder->bIsDuplicateTool = true;
 	RegisterTool(ToolManagerCommands.BeginDuplicateMeshesTool, TEXT("DuplicateMeshesTool"), DuplicateMeshesToolBuilder);
 
+	UGenerateStaticMeshLODAssetToolBuilder* GenerateSMLODToolBuilder = NewObject<UGenerateStaticMeshLODAssetToolBuilder>();
+	GenerateSMLODToolBuilder->AssetAPI = ModelingModeAssetGenerationAPI.Get();
+	RegisterTool(ToolManagerCommands.BeginGenerateStaticMeshLODAssetTool, TEXT("GenerateSMLODTool"), GenerateSMLODToolBuilder);
+
 
 	// edit tools
 
diff --git a/Engine/Plugins/Experimental/ModelingToolsEditorMode/Source/ModelingToolsEditorMode/Private/ModelingToolsEditorModeToolkit.cpp b/Engine/Plugins/Experimental/ModelingToolsEditorMode/Source/ModelingToolsEditorMode/Private/ModelingToolsEditorModeToolkit.cpp
--- a/Engine/Plugins/Experimental/ModelingToolsEditorMode/Source/ModelingToolsEditorMode/Private/ModelingToolsEditorModeToolkit.cpp
+++ b/Engine/Plugins/Experimental/ModelingToolsEditorMode/Source/ModelingToolsEditorMode/Private/ModelingToolsEditorModeToolkit.cpp
@@ -364,6 +364,8 @@ void FModelingToolsEditorModeToolkit::BuildToolPalette_Experimental(FName Palett
 		ToolbarBuilder.AddToolBarButton(Commands.BeginCombineMeshesTool);
 		ToolbarBuilder.AddToolBarButton(Commands.BeginDuplicateMeshesTool);
 		ToolbarBuilder.AddSeparator();
+		ToolbarBuilder.AddToolBarButton(Commands.BeginGenerateStaticMeshLODAssetTool);
+		ToolbarBuilder.AddSeparator();
 		ToolbarBuilder.AddToolBarButton(Commands.BeginMeshInspectorTool);
 	}
 	else if (PaletteIndex == DeformTabName)
diff --git a/Engine/Plugins/Experimental/ModelingToolsEditorMode/Source/ModelingToolsEditorMode/Private/ModelingToolsManagerActions.cpp b/Engine/Plugins/Experimental/ModelingToolsEditorMode/Source/ModelingToolsEditorMode/Private/ModelingToolsManagerActions.cpp
--- a/Engine/Plugins/Experimental/ModelingToolsEditorMode/Source/ModelingToolsEditorMode/Private/ModelingToolsManagerActions.cpp
+++ b/Engine/Plugins/Experimental/ModelingToolsEditorMode/Source/ModelingToolsEditorMode/Private/ModelingToolsManagerActions.cpp
@@ -108,6 +108,7 @@ void FModelingToolsManagerCommands::RegisterCommands()
 	UI_COMMAND(BeginGroomCardsEditorTool, "CardsEditor", "Start the Cards Editor Tool", EUserInterfaceActionType::ToggleButton, FInputChord());
 	UI_COMMAND(BeginSubdividePolyTool, "Subdivide", "Start the SubdividePoly Tool", EUserInterfaceActionType::ToggleButton, FInputChord());
 
+	UI_COMMAND(BeginGenerateStaticMeshLODAssetTool, "AutoLOD", "Start the Generate Static Mesh LOD Asset Tool", EUserInterfaceActionType::ToggleButton, FInputChord());
 
 	UI_COMMAND(BeginPolyModelTool_FaceSelect, "Faces", "Start the PolyGroup Face Selection Tool", EUserInterfaceActionType::ToggleButton, FInputChord());
 	UI_COMMAND(BeginPolyModelTool_EdgeSelect, "Edges", "Start the PolyGroup Edge Selection Tool", EUserInterfaceActionType::ToggleButton, FInputChord());
@@ -122,7 +123,6 @@ void FModelingToolsManagerCommands::RegisterCommands()
 	UI_COMMAND(BeginPolyModelTool_Outset, "Outset", "Start the PolyGroup Outset Tool", EUserInterfaceActionType::ToggleButton, FInputChord());
 	UI_COMMAND(BeginPolyModelTool_CutFaces, "Cut", "Start the PolyGroup Cut Faces Tool", EUserInterfaceActionType::ToggleButton, FInputChord());
 
-
 	UI_COMMAND(AcceptActiveTool, "Accept", "Accept the active tool", EUserInterfaceActionType::Button, FInputChord());
 	UI_COMMAND(CancelActiveTool, "Cancel", "Cancel the active tool", EUserInterfaceActionType::Button, FInputChord());
 	UI_COMMAND(CompleteActiveTool, "Complete", "Complete the active tool", EUserInterfaceActionType::Button, FInputChord());
diff --git a/Engine/Plugins/Experimental/ModelingToolsEditorMode/Source/ModelingToolsEditorMode/Public/ModelingToolsManagerActions.h b/Engine/Plugins/Experimental/ModelingToolsEditorMode/Source/ModelingToolsEditorMode/Public/ModelingToolsManagerActions.h
--- a/Engine/Plugins/Experimental/ModelingToolsEditorMode/Source/ModelingToolsEditorMode/Public/ModelingToolsManagerActions.h
+++ b/Engine/Plugins/Experimental/ModelingToolsEditorMode/Source/ModelingToolsEditorMode/Public/ModelingToolsManagerActions.h
@@ -101,6 +101,7 @@ public:
 	TSharedPtr<FUICommandInfo> BeginGroomMeshToCardsTool;
 	TSharedPtr<FUICommandInfo> BeginGroomCardsEditorTool;
 
+	TSharedPtr<FUICommandInfo> BeginGenerateStaticMeshLODAssetTool;
 
 	TSharedPtr<FUICommandInfo> BeginPolyModelTool_FaceSelect;
 	TSharedPtr<FUICommandInfo> BeginPolyModelTool_EdgeSelect;
