ModelingMode: expose new GroupPaint Tool
#rnx

diff --git a/Engine/Plugins/Experimental/ModelingToolsEditorMode/Source/ModelingToolsEditorMode/Private/ModelingToolsActions.cpp b/Engine/Plugins/Experimental/ModelingToolsEditorMode/Source/ModelingToolsEditorMode/Private/ModelingToolsActions.cpp
--- a/Engine/Plugins/Experimental/ModelingToolsEditorMode/Source/ModelingToolsEditorMode/Private/ModelingToolsActions.cpp
+++ b/Engine/Plugins/Experimental/ModelingToolsEditorMode/Source/ModelingToolsEditorMode/Private/ModelingToolsActions.cpp
@@ -4,6 +4,7 @@
 #include "EditorStyleSet.h"
 #include "DynamicMeshSculptTool.h"
 #include "MeshVertexSculptTool.h"
+#include "MeshGroupPaintTool.h"
 #include "MeshInspectorTool.h"
 #include "DrawPolygonTool.h"
 #include "EditMeshPolygonsTool.h"
@@ -81,6 +82,7 @@ void FModelingToolActionCommands::RegisterAllToolActions()
 	FModelingToolActionCommands::Register();
 	FSculptToolActionCommands::Register();
 	FVertexSculptToolActionCommands::Register();
+	FMeshGroupPaintToolActionCommands::Register();
 	FDrawPolygonToolActionCommands::Register();
 	FTransformToolActionCommands::Register();
 	FMeshSelectionToolActionCommands::Register();
@@ -94,6 +96,7 @@ void FModelingToolActionCommands::UnregisterAllToolActions()
 	FModelingToolActionCommands::Unregister();
 	FSculptToolActionCommands::Unregister();
 	FVertexSculptToolActionCommands::Unregister();
+	FMeshGroupPaintToolActionCommands::Unregister();
 	FDrawPolygonToolActionCommands::Unregister();
 	FTransformToolActionCommands::Unregister();
 	FMeshSelectionToolActionCommands::Unregister();
@@ -121,6 +124,10 @@ void FModelingToolActionCommands::UpdateToolCommandBinding(UInteractiveTool* Too
 	{
 		UPDATE_BINDING(FVertexSculptToolActionCommands);
 	}
+	else if (Cast<UMeshGroupPaintTool>(Tool) != nullptr)
+	{
+		UPDATE_BINDING(FMeshGroupPaintToolActionCommands);
+	}
 	else if (Cast<UDrawPolygonTool>(Tool) != nullptr)
 	{
 		UPDATE_BINDING(FDrawPolygonToolActionCommands);
@@ -164,6 +171,7 @@ void CommandsClassName::GetToolDefaultObjectList(TArray<UInteractiveTool*>& Tool
 
 DEFINE_TOOL_ACTION_COMMANDS(FSculptToolActionCommands, "ModelingToolsSculptTool", "Modeling Tools - Sculpt Tool", UDynamicMeshSculptTool);
 DEFINE_TOOL_ACTION_COMMANDS(FVertexSculptToolActionCommands, "ModelingToolsVertexSculptTool", "Modeling Tools - Vertex Sculpt Tool", UMeshVertexSculptTool);
+DEFINE_TOOL_ACTION_COMMANDS(FMeshGroupPaintToolActionCommands, "ModelingToolsMeshGroupPaintTool", "Modeling Tools - Group Paint Tool", UMeshGroupPaintTool);
 DEFINE_TOOL_ACTION_COMMANDS(FTransformToolActionCommands, "ModelingToolsTransformTool", "Modeling Tools - Transform Tool", UTransformMeshesTool);
 DEFINE_TOOL_ACTION_COMMANDS(FDrawPolygonToolActionCommands, "ModelingToolsDrawPolygonTool", "Modeling Tools - Draw Polygon Tool", UDrawPolygonTool);
 DEFINE_TOOL_ACTION_COMMANDS(FMeshSelectionToolActionCommands, "ModelingToolsMeshSelectionTool", "Modeling Tools - Mesh Selection Tool", UMeshSelectionTool);
diff --git a/Engine/Plugins/Experimental/ModelingToolsEditorMode/Source/ModelingToolsEditorMode/Private/ModelingToolsEditorMode.cpp b/Engine/Plugins/Experimental/ModelingToolsEditorMode/Source/ModelingToolsEditorMode/Private/ModelingToolsEditorMode.cpp
--- a/Engine/Plugins/Experimental/ModelingToolsEditorMode/Source/ModelingToolsEditorMode/Private/ModelingToolsEditorMode.cpp
+++ b/Engine/Plugins/Experimental/ModelingToolsEditorMode/Source/ModelingToolsEditorMode/Private/ModelingToolsEditorMode.cpp
@@ -74,6 +74,7 @@
 #include "ProjectToTargetTool.h"
 #include "LatticeDeformerTool.h"
 #include "SeamSculptTool.h"
+#include "MeshGroupPaintTool.h"
 
 #include "Physics/PhysicsInspectorTool.h"
 #include "Physics/SetCollisionGeometryTool.h"
@@ -382,6 +383,10 @@ void UModelingToolsEditorMode::Enter()
 	MoveVerticesToolBuilder->StylusAPI = StylusStateTracker.Get();
 	RegisterTool(ToolManagerCommands.BeginSculptMeshTool, TEXT("MoveVerticesTool"), MoveVerticesToolBuilder);
 
+	auto MeshGroupPaintToolBuilder = NewObject<UMeshGroupPaintToolBuilder>();
+	MeshGroupPaintToolBuilder->StylusAPI = StylusStateTracker.Get();
+	RegisterTool(ToolManagerCommands.BeginMeshGroupPaintTool, TEXT("MeshGroupPaintTool"), MeshGroupPaintToolBuilder);
+
 	RegisterTool(ToolManagerCommands.BeginPolyEditTool, TEXT("EditMeshPolygonsTool"), NewObject<UEditMeshPolygonsToolBuilder>());
 	UEditMeshPolygonsToolBuilder* TriEditBuilder = NewObject<UEditMeshPolygonsToolBuilder>();
 	TriEditBuilder->bTriangleMode = true;
diff --git a/Engine/Plugins/Experimental/ModelingToolsEditorMode/Source/ModelingToolsEditorMode/Private/ModelingToolsEditorModeToolkit.cpp b/Engine/Plugins/Experimental/ModelingToolsEditorMode/Source/ModelingToolsEditorMode/Private/ModelingToolsEditorModeToolkit.cpp
--- a/Engine/Plugins/Experimental/ModelingToolsEditorMode/Source/ModelingToolsEditorMode/Private/ModelingToolsEditorModeToolkit.cpp
+++ b/Engine/Plugins/Experimental/ModelingToolsEditorMode/Source/ModelingToolsEditorMode/Private/ModelingToolsEditorModeToolkit.cpp
@@ -413,6 +413,7 @@ void FModelingToolsEditorModeToolkit::BuildToolPalette_Experimental(FName Palett
 		ToolbarBuilder.AddToolBarButton(Commands.BeginTransformMeshesTool);
 		ToolbarBuilder.AddToolBarButton(Commands.BeginMeshSelectionTool);
 		ToolbarBuilder.AddToolBarButton(Commands.BeginPolyGroupsTool);
+		ToolbarBuilder.AddToolBarButton(Commands.BeginMeshGroupPaintTool);
 		ToolbarBuilder.AddSeparator();
 		ToolbarBuilder.AddToolBarButton(Commands.BeginPolyEditTool);
 		ToolbarBuilder.AddToolBarButton(Commands.BeginPolyDeformTool);
diff --git a/Engine/Plugins/Experimental/ModelingToolsEditorMode/Source/ModelingToolsEditorMode/Private/ModelingToolsManagerActions.cpp b/Engine/Plugins/Experimental/ModelingToolsEditorMode/Source/ModelingToolsEditorMode/Private/ModelingToolsManagerActions.cpp
--- a/Engine/Plugins/Experimental/ModelingToolsEditorMode/Source/ModelingToolsEditorMode/Private/ModelingToolsManagerActions.cpp
+++ b/Engine/Plugins/Experimental/ModelingToolsEditorMode/Source/ModelingToolsEditorMode/Private/ModelingToolsManagerActions.cpp
@@ -110,6 +110,8 @@ void FModelingToolsManagerCommands::RegisterCommands()
 	UI_COMMAND(BeginSubdividePolyTool, "Subdivide", "Start the SubdividePoly Tool", EUserInterfaceActionType::ToggleButton, FInputChord());
 
 	UI_COMMAND(BeginGenerateStaticMeshLODAssetTool, "AutoLOD", "Start the Generate Static Mesh LOD Asset Tool", EUserInterfaceActionType::ToggleButton, FInputChord());
+	UI_COMMAND(BeginMeshGroupPaintTool, "GroupPaint", "Start the Mesh PolyGroup Paint Tool", EUserInterfaceActionType::ToggleButton, FInputChord());
+
 
 	UI_COMMAND(BeginPolyModelTool_FaceSelect, "Faces", "Start the PolyGroup Face Selection Tool", EUserInterfaceActionType::ToggleButton, FInputChord());
 	UI_COMMAND(BeginPolyModelTool_EdgeSelect, "Edges", "Start the PolyGroup Edge Selection Tool", EUserInterfaceActionType::ToggleButton, FInputChord());
diff --git a/Engine/Plugins/Experimental/ModelingToolsEditorMode/Source/ModelingToolsEditorMode/Public/ModelingToolsActions.h b/Engine/Plugins/Experimental/ModelingToolsEditorMode/Source/ModelingToolsEditorMode/Public/ModelingToolsActions.h
--- a/Engine/Plugins/Experimental/ModelingToolsEditorMode/Source/ModelingToolsEditorMode/Public/ModelingToolsActions.h
+++ b/Engine/Plugins/Experimental/ModelingToolsEditorMode/Source/ModelingToolsEditorMode/Public/ModelingToolsActions.h
@@ -77,6 +77,7 @@ public:\
 
 DECLARE_TOOL_ACTION_COMMANDS(FSculptToolActionCommands);
 DECLARE_TOOL_ACTION_COMMANDS(FVertexSculptToolActionCommands);
+DECLARE_TOOL_ACTION_COMMANDS(FMeshGroupPaintToolActionCommands);
 DECLARE_TOOL_ACTION_COMMANDS(FTransformToolActionCommands);
 DECLARE_TOOL_ACTION_COMMANDS(FDrawPolygonToolActionCommands);
 DECLARE_TOOL_ACTION_COMMANDS(FMeshSelectionToolActionCommands);
diff --git a/Engine/Plugins/Experimental/ModelingToolsEditorMode/Source/ModelingToolsEditorMode/Public/ModelingToolsManagerActions.h b/Engine/Plugins/Experimental/ModelingToolsEditorMode/Source/ModelingToolsEditorMode/Public/ModelingToolsManagerActions.h
--- a/Engine/Plugins/Experimental/ModelingToolsEditorMode/Source/ModelingToolsEditorMode/Public/ModelingToolsManagerActions.h
+++ b/Engine/Plugins/Experimental/ModelingToolsEditorMode/Source/ModelingToolsEditorMode/Public/ModelingToolsManagerActions.h
@@ -78,6 +78,7 @@ public:
 	TSharedPtr<FUICommandInfo> BeginBspConversionTool;
 	TSharedPtr<FUICommandInfo> BeginMeshToVolumeTool;
 	TSharedPtr<FUICommandInfo> BeginVolumeToMeshTool;
+	TSharedPtr<FUICommandInfo> BeginMeshGroupPaintTool;
 
 	TSharedPtr<FUICommandInfo> BeginPhysicsInspectorTool;
 	TSharedPtr<FUICommandInfo> BeginSetCollisionGeometryTool;
