MeshModelingTools: First pass of reorganizing the tool tabs.

#rnx

diff --git a/Engine/Plugins/Experimental/ModelingToolsEditorMode/Source/ModelingToolsEditorMode/Private/ModelingToolsEditorModeToolkit.cpp b/Engine/Plugins/Experimental/ModelingToolsEditorMode/Source/ModelingToolsEditorMode/Private/ModelingToolsEditorModeToolkit.cpp
--- a/Engine/Plugins/Experimental/ModelingToolsEditorMode/Source/ModelingToolsEditorMode/Private/ModelingToolsEditorModeToolkit.cpp
+++ b/Engine/Plugins/Experimental/ModelingToolsEditorMode/Source/ModelingToolsEditorMode/Private/ModelingToolsEditorModeToolkit.cpp
@@ -31,7 +31,7 @@
 static TAutoConsoleVariable<int32> CVarEnablePolyModeling(
 	TEXT("modeling.EnablePolyModel"),
 	1,
-	TEXT("Enable prototype PolyModeling tab"));
+	TEXT("Enable prototype PolyEdit tab"));
 
 
 
@@ -267,21 +267,21 @@
 
 static const FName PrimitiveTabName(TEXT("Primitives"));
 static const FName CreateTabName(TEXT("Create"));
-static const FName EditTabName(TEXT("Edit"));
+static const FName AttributesTabName(TEXT("Attributes"));
 static const FName SculptTabName(TEXT("Sculpt"));
 static const FName TrianglesTabName(TEXT("Triangles"));
-static const FName PolyGroupsTabName(TEXT("PolyGroups"));
-static const FName UVNormalTabName(TEXT("UVs/Normals"));
+static const FName PolyModelingTabName(TEXT("PolyModeling"));
+static const FName UVTabName(TEXT("UVs"));
 static const FName TransformTabName(TEXT("Transform"));
-static const FName DeformTabName(TEXT("Deform"));
+static const FName ModifyTabName(TEXT("Modify"));
 static const FName VolumesTabName(TEXT("Volumes"));
 static const FName PrototypesTabName(TEXT("Prototypes"));
 static const FName HairTabName(TEXT("Hair"));
-static const FName PolyModelTabName(TEXT("PolyModel"));
+static const FName PolyEditTabName(TEXT("PolyEdit"));
 
 
-const TArray<FName> FModelingToolsEditorModeToolkit::PaletteNames_Standard = { PrimitiveTabName, CreateTabName, TransformTabName, DeformTabName, PolyGroupsTabName, TrianglesTabName, UVNormalTabName, VolumesTabName, HairTabName };
-const TArray<FName> FModelingToolsEditorModeToolkit::PaletteNames_Experimental = { PrimitiveTabName, CreateTabName, TransformTabName, DeformTabName, PolyGroupsTabName, TrianglesTabName, UVNormalTabName, VolumesTabName, HairTabName, PrototypesTabName };
+const TArray<FName> FModelingToolsEditorModeToolkit::PaletteNames_Standard = { PrimitiveTabName, CreateTabName, PolyModelingTabName, TrianglesTabName, ModifyTabName, TransformTabName, AttributesTabName, UVTabName, VolumesTabName, HairTabName };
+const TArray<FName> FModelingToolsEditorModeToolkit::PaletteNames_Experimental = { PrimitiveTabName, CreateTabName, PolyModelingTabName, TrianglesTabName, ModifyTabName, TransformTabName, AttributesTabName, UVTabName, VolumesTabName, HairTabName, PrototypesTabName };
 
 
 void FModelingToolsEditorModeToolkit::GetToolPaletteNames(TArray<FName>& PaletteNames) const
@@ -292,7 +292,7 @@
 	bool bEnablePolyModel = (CVarEnablePolyModeling.GetValueOnGameThread() > 0);
 	if (bEnablePolyModel)
 	{
-		PaletteNames.Add(PolyModelTabName);
+		PaletteNames.Add(PolyEditTabName);
 	}
 }
 
@@ -336,8 +336,6 @@
 	}
 	else if (PaletteIndex == CreateTabName)
 	{
-		ToolbarBuilder.AddToolBarButton(Commands.BeginTransformMeshesTool);
-		ToolbarBuilder.AddToolBarButton(Commands.BeginMeshSelectionTool);
 		ToolbarBuilder.AddSeparator();
 		ToolbarBuilder.AddToolBarButton(Commands.BeginDrawPolygonTool);
 		ToolbarBuilder.AddToolBarButton(Commands.BeginDrawPolyPathTool);
@@ -356,24 +354,19 @@
 	}
 	else if (PaletteIndex == TransformTabName)
 	{
-		ToolbarBuilder.AddToolBarButton(Commands.BeginTransformMeshesTool);
-		ToolbarBuilder.AddToolBarButton(Commands.BeginMeshSelectionTool);
-		ToolbarBuilder.AddSeparator();
-		ToolbarBuilder.AddToolBarButton(Commands.BeginAlignObjectsTool);
 		ToolbarBuilder.AddToolBarButton(Commands.BeginEditPivotTool);
 		ToolbarBuilder.AddToolBarButton(Commands.BeginBakeTransformTool);
+		ToolbarBuilder.AddSeparator();
 		ToolbarBuilder.AddToolBarButton(Commands.BeginCombineMeshesTool);
 		ToolbarBuilder.AddToolBarButton(Commands.BeginDuplicateMeshesTool);
 		ToolbarBuilder.AddSeparator();
 		ToolbarBuilder.AddToolBarButton(Commands.BeginGenerateStaticMeshLODAssetTool);
 		ToolbarBuilder.AddSeparator();
-		ToolbarBuilder.AddToolBarButton(Commands.BeginMeshInspectorTool);
+		ToolbarBuilder.AddToolBarButton(Commands.BeginTransformMeshesTool);
+		ToolbarBuilder.AddToolBarButton(Commands.BeginAlignObjectsTool);
 	}
-	else if (PaletteIndex == DeformTabName)
+	else if (PaletteIndex == ModifyTabName)
 	{
-		ToolbarBuilder.AddToolBarButton(Commands.BeginTransformMeshesTool);
-		ToolbarBuilder.AddToolBarButton(Commands.BeginMeshSelectionTool);
-		ToolbarBuilder.AddSeparator();
 		ToolbarBuilder.AddToolBarButton(Commands.BeginSculptMeshTool);
 		ToolbarBuilder.AddToolBarButton(Commands.BeginRemeshSculptMeshTool);
 		ToolbarBuilder.AddSeparator();
@@ -385,15 +378,10 @@
 		ToolbarBuilder.AddToolBarButton(Commands.BeginDisplaceMeshTool);
 		ToolbarBuilder.AddToolBarButton(Commands.BeginMeshSpaceDeformerTool);
 		ToolbarBuilder.AddToolBarButton(Commands.BeginLatticeDeformerTool);
-		ToolbarBuilder.AddSeparator();
-		ToolbarBuilder.AddToolBarButton(Commands.BeginMeshAttributePaintTool);
-		ToolbarBuilder.AddSeparator();
 		ToolbarBuilder.AddToolBarButton(Commands.BeginRemeshMeshTool);
-		ToolbarBuilder.AddToolBarButton(Commands.BeginMeshInspectorTool);
 	}
 	else if (PaletteIndex == TrianglesTabName)
 	{
-		ToolbarBuilder.AddToolBarButton(Commands.BeginTransformMeshesTool);
 		ToolbarBuilder.AddToolBarButton(Commands.BeginMeshSelectionTool);
 		ToolbarBuilder.AddSeparator();
 		ToolbarBuilder.AddToolBarButton(Commands.BeginTriEditTool);
@@ -405,34 +393,30 @@
 		ToolbarBuilder.AddToolBarButton(Commands.BeginWeldEdgesTool);
 		ToolbarBuilder.AddToolBarButton(Commands.BeginRemoveOccludedTrianglesTool);
 		ToolbarBuilder.AddToolBarButton(Commands.BeginHoleFillTool);
-		ToolbarBuilder.AddSeparator();
-		ToolbarBuilder.AddToolBarButton(Commands.BeginMeshInspectorTool);
 	}
-	else if (PaletteIndex == PolyGroupsTabName)
+	else if (PaletteIndex == PolyModelingTabName)
 	{
-		ToolbarBuilder.AddToolBarButton(Commands.BeginTransformMeshesTool);
-		ToolbarBuilder.AddToolBarButton(Commands.BeginMeshSelectionTool);
-		ToolbarBuilder.AddToolBarButton(Commands.BeginPolyGroupsTool);
-		ToolbarBuilder.AddToolBarButton(Commands.BeginMeshGroupPaintTool);
-		ToolbarBuilder.AddSeparator();
 		ToolbarBuilder.AddToolBarButton(Commands.BeginPolyEditTool);
 		ToolbarBuilder.AddToolBarButton(Commands.BeginPolyDeformTool);
 		ToolbarBuilder.AddSeparator();
 		ToolbarBuilder.AddToolBarButton(Commands.BeginGroupEdgeInsertionTool);
 		ToolbarBuilder.AddToolBarButton(Commands.BeginEdgeLoopInsertionTool);
 		ToolbarBuilder.AddSeparator();
+		ToolbarBuilder.AddToolBarButton(Commands.BeginDrawPolygonTool);
+		ToolbarBuilder.AddToolBarButton(Commands.BeginMeshBooleanTool);
+	}
+	else if (PaletteIndex == AttributesTabName)
+	{
+		ToolbarBuilder.AddToolBarButton(Commands.BeginPolyGroupsTool);
 		ToolbarBuilder.AddToolBarButton(Commands.BeginAttributeEditorTool);
 		ToolbarBuilder.AddToolBarButton(Commands.BeginMeshInspectorTool);
-	}
-	else if (PaletteIndex == UVNormalTabName)
-	{
-		ToolbarBuilder.AddToolBarButton(Commands.BeginTransformMeshesTool);
-		ToolbarBuilder.AddToolBarButton(Commands.BeginMeshSelectionTool);
-		ToolbarBuilder.AddToolBarButton(Commands.BeginPolyGroupsTool);
-		ToolbarBuilder.AddSeparator();
+		ToolbarBuilder.AddToolBarButton(Commands.BeginMeshGroupPaintTool);
 		ToolbarBuilder.AddToolBarButton(Commands.BeginEditNormalsTool);
 		ToolbarBuilder.AddToolBarButton(Commands.BeginEditTangentsTool);
-		ToolbarBuilder.AddSeparator();
+		ToolbarBuilder.AddToolBarButton(Commands.BeginMeshAttributePaintTool);
+	}
+	else if (PaletteIndex == UVTabName)
+	{
 		ToolbarBuilder.AddToolBarButton(Commands.BeginGlobalUVGenerateTool);
 		ToolbarBuilder.AddToolBarButton(Commands.BeginGroupUVGenerateTool);
 		ToolbarBuilder.AddToolBarButton(Commands.BeginUVProjectionTool);
@@ -441,9 +425,6 @@
 		ToolbarBuilder.AddToolBarButton(Commands.BeginUVLayoutTool);
 		ToolbarBuilder.AddSeparator();
 		ToolbarBuilder.AddToolBarButton(Commands.BeginBakeMeshAttributeMapsTool);
-		ToolbarBuilder.AddSeparator();
-		ToolbarBuilder.AddToolBarButton(Commands.BeginAttributeEditorTool);
-		ToolbarBuilder.AddToolBarButton(Commands.BeginMeshInspectorTool);
 	}
 	else if (PaletteIndex == VolumesTabName)
 	{
@@ -459,26 +440,17 @@
 	}
 	else if (PaletteIndex == HairTabName)
 	{
-		ToolbarBuilder.AddToolBarButton(Commands.BeginTransformMeshesTool);
-		ToolbarBuilder.AddToolBarButton(Commands.BeginMeshSelectionTool);
-		ToolbarBuilder.AddSeparator();
 		ToolbarBuilder.AddToolBarButton(Commands.BeginGroomToMeshTool);
 		ToolbarBuilder.AddToolBarButton(Commands.BeginGroomCardsEditorTool);
 		ToolbarBuilder.AddToolBarButton(Commands.BeginGenerateLODMeshesTool);
-		ToolbarBuilder.AddSeparator();
-		ToolbarBuilder.AddToolBarButton(Commands.BeginAttributeEditorTool);
-		ToolbarBuilder.AddToolBarButton(Commands.BeginMeshInspectorTool);
 	}
 	else if (PaletteIndex == PrototypesTabName)
 	{
-		ToolbarBuilder.AddToolBarButton(Commands.BeginTransformMeshesTool);
-		ToolbarBuilder.AddToolBarButton(Commands.BeginMeshSelectionTool);
-		ToolbarBuilder.AddSeparator();
 		ToolbarBuilder.AddToolBarButton(Commands.BeginAddPatchTool);
 		ToolbarBuilder.AddToolBarButton(Commands.BeginShapeSprayTool);
 		ToolbarBuilder.AddToolBarButton(Commands.BeginSubdividePolyTool);
 	}
-	else if (PaletteIndex == PolyModelTabName)
+	else if (PaletteIndex == PolyEditTabName)
 	{
 		ToolbarBuilder.AddToolBarButton(Commands.BeginPolyModelTool_FaceSelect);
 		ToolbarBuilder.AddToolBarButton(Commands.BeginPolyModelTool_EdgeSelect);
diff --git a/Engine/Plugins/Experimental/ModelingToolsEditorMode/Source/ModelingToolsEditorMode/Private/ModelingToolsManagerActions.cpp b/Engine/Plugins/Experimental/ModelingToolsEditorMode/Source/ModelingToolsEditorMode/Private/ModelingToolsManagerActions.cpp
--- a/Engine/Plugins/Experimental/ModelingToolsEditorMode/Source/ModelingToolsEditorMode/Private/ModelingToolsManagerActions.cpp
+++ b/Engine/Plugins/Experimental/ModelingToolsEditorMode/Source/ModelingToolsEditorMode/Private/ModelingToolsManagerActions.cpp
@@ -85,7 +85,7 @@ void FModelingToolsManagerCommands::RegisterCommands()
 	UI_COMMAND(BeginPolygonCutTool, "PolyCut", "Start the Polygon Cut Tool", EUserInterfaceActionType::ToggleButton, FInputChord());
 	UI_COMMAND(BeginGlobalUVGenerateTool, "AutoUnwrap", "Start the Global UV Unwrap Tool", EUserInterfaceActionType::ToggleButton, FInputChord());
 	UI_COMMAND(BeginGroupUVGenerateTool, "GroupUnwrap", "Stat the PolyGroup UV Unwrap Tool", EUserInterfaceActionType::ToggleButton, FInputChord());
-	UI_COMMAND(BeginMeshSelectionTool, "Select", "Start the Mesh Selection Tool", EUserInterfaceActionType::ToggleButton, FInputChord());
+	UI_COMMAND(BeginMeshSelectionTool, "TriSelect", "Start the Triangle Selection Tool", EUserInterfaceActionType::ToggleButton, FInputChord());
 
 	UI_COMMAND(BeginPhysicsInspectorTool, "PhysInspect", "Inspect Physics Geometry for selected Meshes", EUserInterfaceActionType::ToggleButton, FInputChord());
 	UI_COMMAND(BeginSetCollisionGeometryTool, "MeshToCollision", "Convert selected Meshes to Collision Geometry", EUserInterfaceActionType::ToggleButton, FInputChord());
