ModelingTools: add FModelingToolsAxisFilter UStruct, use in UMoveBrushOpProps to allow brush effect to be constrained by world-space axis
ModelingMode: add details customization for FModelingToolsAxisFilter so that it shows X/Y/Z checkboxes horizontally
#rnx

diff --git a/Engine/Plugins/Experimental/MeshModelingToolset/Source/MeshModelingTools/Public/Sculpting/MeshMoveBrushOps.h b/Engine/Plugins/Experimental/MeshModelingToolset/Source/MeshModelingTools/Public/Sculpting/MeshMoveBrushOps.h
--- a/Engine/Plugins/Experimental/MeshModelingToolset/Source/MeshModelingTools/Public/Sculpting/MeshMoveBrushOps.h
+++ b/Engine/Plugins/Experimental/MeshModelingToolset/Source/MeshModelingTools/Public/Sculpting/MeshMoveBrushOps.h
@@ -7,9 +7,11 @@
 #include "MeshNormals.h"
 #include "MeshWeights.h"
 #include "Async/ParallelFor.h"
+#include "PropertySets/AxisFilterPropertyType.h"
 #include "MeshMoveBrushOps.generated.h"
 
 
+
 UCLASS()
 class MESHMODELINGTOOLS_API UMoveBrushOpProps : public UMeshSculptBrushOpProps
 {
@@ -30,6 +32,10 @@ public:
 	virtual float GetStrength() override { return Strength; }
 	virtual float GetFalloff() override { return Falloff; }
 	virtual float GetDepth() override { return Depth; }
+
+	/** Axis filters restrict mesh movement to World X/Y/Z axes */
+	UPROPERTY(EditAnywhere, Category = MoveBrush, AdvancedDisplay)
+	FModelingToolsAxisFilter AxisFilters;
 };
 
 
@@ -43,6 +49,20 @@ public:
 
 		double UsePower = Stamp.Power;
 		FVector3d MoveVec = Stamp.LocalFrame.Origin - Stamp.PrevLocalFrame.Origin;
+		
+		UMoveBrushOpProps* Props = GetPropertySetAs<UMoveBrushOpProps>();
+
+		// if we have axis constraints we want to apply them in world space, and then remap
+		// that world vector to local space
+		if (Props->AxisFilters.AnyAxisFiltered())
+		{
+			FVector3d WorldMoveVec = Stamp.WorldFrame.Origin - Stamp.PrevWorldFrame.Origin;
+			WorldMoveVec = MoveVec.Length() * WorldMoveVec.Normalized();		// apply local space scaling
+			WorldMoveVec.X = (Props->AxisFilters.bAxisX) ? WorldMoveVec.X : 0;
+			WorldMoveVec.Y = (Props->AxisFilters.bAxisY) ? WorldMoveVec.Y : 0;
+			WorldMoveVec.Z = (Props->AxisFilters.bAxisZ) ? WorldMoveVec.Z : 0;
+			MoveVec = Stamp.LocalFrame.FromFrameVector(Stamp.WorldFrame.ToFrameVector(WorldMoveVec));
+		}
 
 		ParallelFor(Vertices.Num(), [&](int32 k)
 		{
diff --git a/Engine/Plugins/Experimental/MeshModelingToolset/Source/ModelingComponents/Public/PropertySets/AxisFilterPropertyType.h b/Engine/Plugins/Experimental/MeshModelingToolset/Source/ModelingComponents/Public/PropertySets/AxisFilterPropertyType.h
new file mode 100644
--- /dev/null
+++ b/Engine/Plugins/Experimental/MeshModelingToolset/Source/ModelingComponents/Public/PropertySets/AxisFilterPropertyType.h
@@ -0,0 +1,30 @@
+// Copyright Epic Games, Inc. All Rights Reserved.
+
+#pragma once
+
+#include "CoreMinimal.h"
+#include "AxisFilterPropertyType.generated.h"
+
+
+USTRUCT(BlueprintType)
+struct MODELINGCOMPONENTS_API FModelingToolsAxisFilter
+{
+	GENERATED_BODY()
+
+	/** X Axis */
+	UPROPERTY(EditAnywhere, Category = AxisFilters, meta = (DisplayName = "X"))
+	bool bAxisX = true;
+
+	/** Y Axis */
+	UPROPERTY(EditAnywhere, Category = AxisFilters, meta = (DisplayName = "Y"))
+	bool bAxisY = true;
+
+	/** Z Axis */
+	UPROPERTY(EditAnywhere, Category = AxisFilters, meta = (DisplayName = "Z"))
+	bool bAxisZ = true;
+
+
+	// helper functions
+
+	bool AnyAxisFiltered() const { return bAxisX == false || bAxisY == false || bAxisZ == false; }
+};
diff --git a/Engine/Plugins/Experimental/ModelingToolsEditorMode/Source/ModelingToolsEditorMode/Private/DetailsCustomizations/ModelingToolsAxisFilterCustomization.cpp b/Engine/Plugins/Experimental/ModelingToolsEditorMode/Source/ModelingToolsEditorMode/Private/DetailsCustomizations/ModelingToolsAxisFilterCustomization.cpp
new file mode 100644
--- /dev/null
+++ b/Engine/Plugins/Experimental/ModelingToolsEditorMode/Source/ModelingToolsEditorMode/Private/DetailsCustomizations/ModelingToolsAxisFilterCustomization.cpp
@@ -0,0 +1,101 @@
+// Copyright Epic Games, Inc. All Rights Reserved.
+
+#include "DetailsCustomizations/ModelingToolPropertyCustomizations.h"
+
+#include "PropertyHandle.h"
+#include "IPropertyTypeCustomization.h"
+#include "IPropertyUtilities.h"
+#include "DetailWidgetRow.h"
+#include "EditorFontGlyphs.h"
+
+#include "PropertySets/AxisFilterPropertyType.h"
+
+
+#define LOCTEXT_NAMESPACE "ModelingToolsAxisFilterCustomization"
+
+
+TSharedRef<IPropertyTypeCustomization> FModelingToolsAxisFilterCustomization::MakeInstance()
+{
+	return MakeShareable(new FModelingToolsAxisFilterCustomization);
+}
+
+void FModelingToolsAxisFilterCustomization::CustomizeHeader(TSharedRef<IPropertyHandle> StructPropertyHandle, FDetailWidgetRow& HeaderRow, IPropertyTypeCustomizationUtils& StructCustomizationUtils)
+{
+	const float XYZPadding = 10.f;
+
+	TSharedPtr<IPropertyHandle> bAxisX = StructPropertyHandle->GetChildHandle(GET_MEMBER_NAME_CHECKED(FModelingToolsAxisFilter, bAxisX));
+	bAxisX->MarkHiddenByCustomization();
+
+	TSharedPtr<IPropertyHandle> bAxisY = StructPropertyHandle->GetChildHandle(GET_MEMBER_NAME_CHECKED(FModelingToolsAxisFilter, bAxisY));
+	bAxisY->MarkHiddenByCustomization();
+
+	TSharedPtr<IPropertyHandle> bAxisZ = StructPropertyHandle->GetChildHandle(GET_MEMBER_NAME_CHECKED(FModelingToolsAxisFilter, bAxisZ));
+	bAxisZ->MarkHiddenByCustomization();
+
+	HeaderRow.NameContent()
+		[
+			StructPropertyHandle->CreatePropertyNameWidget()
+		]
+	.ValueContent()
+		.MinDesiredWidth(125.f)
+		[
+			SNew(SHorizontalBox)
+			+ SHorizontalBox::Slot()
+			.Padding(0.f, 0.f, XYZPadding, 0.f)
+			.AutoWidth()
+			[
+				SNew(SHorizontalBox)
+				+ SHorizontalBox::Slot()
+				.AutoWidth()
+				[
+					bAxisX->CreatePropertyNameWidget()
+				]
+				+ SHorizontalBox::Slot()
+				.AutoWidth()
+				[
+					bAxisX->CreatePropertyValueWidget()
+				]
+			]
+
+			+ SHorizontalBox::Slot()
+			.Padding(0.f, 0.f, XYZPadding, 0.f)
+			.AutoWidth()
+			[
+				SNew(SHorizontalBox)
+				+ SHorizontalBox::Slot()
+				.AutoWidth()
+				[
+					bAxisY->CreatePropertyNameWidget()
+				]
+				+ SHorizontalBox::Slot()
+				.AutoWidth()
+				[
+					bAxisY->CreatePropertyValueWidget()
+				]
+			]
+
+			+ SHorizontalBox::Slot()
+			.Padding(0.f, 0.f, XYZPadding, 0.f)
+			.AutoWidth()
+			[
+				SNew(SHorizontalBox)
+				+ SHorizontalBox::Slot()
+				.AutoWidth()
+				[
+					bAxisZ->CreatePropertyNameWidget()
+				]
+				+ SHorizontalBox::Slot()
+				.AutoWidth()
+				[
+					bAxisZ->CreatePropertyValueWidget()
+				]
+			]
+		];
+}
+
+void FModelingToolsAxisFilterCustomization::CustomizeChildren(TSharedRef<IPropertyHandle> StructPropertyHandle, class IDetailChildrenBuilder& ChildBuilder, IPropertyTypeCustomizationUtils& StructCustomizationUtils)
+{
+}
+
+
+#undef LOCTEXT_NAMESPACE
\ No newline at end of file
diff --git a/Engine/Plugins/Experimental/ModelingToolsEditorMode/Source/ModelingToolsEditorMode/Private/ModelingToolsEditorModeModule.cpp b/Engine/Plugins/Experimental/ModelingToolsEditorMode/Source/ModelingToolsEditorMode/Private/ModelingToolsEditorModeModule.cpp
--- a/Engine/Plugins/Experimental/ModelingToolsEditorMode/Source/ModelingToolsEditorMode/Private/ModelingToolsEditorModeModule.cpp
+++ b/Engine/Plugins/Experimental/ModelingToolsEditorMode/Source/ModelingToolsEditorMode/Private/ModelingToolsEditorModeModule.cpp
@@ -11,6 +11,10 @@
 #include "ISettingsSection.h"
 #include "Misc/CoreDelegates.h"
 
+#include "Modules/ModuleManager.h"
+#include "PropertyEditorModule.h"
+#include "DetailsCustomizations/ModelingToolPropertyCustomizations.h"
+
 #define LOCTEXT_NAMESPACE "FModelingToolsEditorModeModule"
 
 void FModelingToolsEditorModeModule::StartupModule()
@@ -56,6 +60,9 @@ void FModelingToolsEditorModeModule::OnPostEngineInit()
 		);
 	}
 
+	// Register details view customizations
+	FPropertyEditorModule& PropertyModule = FModuleManager::LoadModuleChecked<FPropertyEditorModule>("PropertyEditor");
+	PropertyModule.RegisterCustomPropertyTypeLayout("ModelingToolsAxisFilter", FOnGetPropertyTypeCustomizationInstance::CreateStatic(&FModelingToolsAxisFilterCustomization::MakeInstance));
 }
 
 #undef LOCTEXT_NAMESPACE
diff --git a/Engine/Plugins/Experimental/ModelingToolsEditorMode/Source/ModelingToolsEditorMode/Public/DetailsCustomizations/ModelingToolPropertyCustomizations.h b/Engine/Plugins/Experimental/ModelingToolsEditorMode/Source/ModelingToolsEditorMode/Public/DetailsCustomizations/ModelingToolPropertyCustomizations.h
new file mode 100644
--- /dev/null
+++ b/Engine/Plugins/Experimental/ModelingToolsEditorMode/Source/ModelingToolsEditorMode/Public/DetailsCustomizations/ModelingToolPropertyCustomizations.h
@@ -0,0 +1,20 @@
+// Copyright Epic Games, Inc. All Rights Reserved.
+
+#pragma once
+
+#include "IPropertyTypeCustomization.h"
+#include "Input/Reply.h"
+#include "Widgets/Text/STextBlock.h"
+#include "Widgets/Input/SButton.h"
+
+class IDetailLayoutBuilder;
+
+/** Details customization for FModelingModeAxisFilter struct - X/Y/Z booleans are laid out in a single horizontal row */
+class FModelingToolsAxisFilterCustomization : public IPropertyTypeCustomization
+{
+public:
+	static TSharedRef<IPropertyTypeCustomization> MakeInstance();
+
+	virtual void CustomizeHeader(TSharedRef<IPropertyHandle> StructPropertyHandle, class FDetailWidgetRow& HeaderRow, IPropertyTypeCustomizationUtils& StructCustomizationUtils) override;
+	virtual void CustomizeChildren(TSharedRef<IPropertyHandle> StructPropertyHandle, class IDetailChildrenBuilder& ChildBuilder, IPropertyTypeCustomizationUtils& StructCustomizationUtils) override;
+};
