MeshModelingTools: Made tools in modeling mode register their hotkeys with toolkit commands rather than tool commands, to allow them to work when the viewport is not focused.

#rnx

diff --git a/Engine/Plugins/Experimental/ModelingToolsEditorMode/Source/ModelingToolsEditorMode/Private/ModelingToolsEditorMode.cpp b/Engine/Plugins/Experimental/ModelingToolsEditorMode/Source/ModelingToolsEditorMode/Private/ModelingToolsEditorMode.cpp
--- a/Engine/Plugins/Experimental/ModelingToolsEditorMode/Source/ModelingToolsEditorMode/Private/ModelingToolsEditorMode.cpp
+++ b/Engine/Plugins/Experimental/ModelingToolsEditorMode/Source/ModelingToolsEditorMode/Private/ModelingToolsEditorMode.cpp
@@ -604,8 +604,9 @@ void UModelingToolsEditorMode::Enter()
 
 	ToolsContext->ToolManager->SelectActiveToolType(EToolSide::Left, TEXT("DynaSculptTool"));
 
-	// register modeling mode hotkeys
-	FModelingModeActionCommands::RegisterCommandBindings(ToolCommandList, [this](EModelingModeActionCommands Command) {
+	// Register modeling mode hotkeys. Note that we use the toolkit command list because we would like the hotkeys
+	// to work even when the viewport is not focused, provided that nothing else captures the key presses.
+	FModelingModeActionCommands::RegisterCommandBindings(Toolkit->GetToolkitCommands(), [this](EModelingModeActionCommands Command) {
 		ModelingModeShortcutRequested(Command);
 	});
 
@@ -635,7 +636,7 @@ void UModelingToolsEditorMode::Exit()
 
 	StylusStateTracker = nullptr;
 
-	FModelingModeActionCommands::UnRegisterCommandBindings(ToolCommandList);
+	FModelingModeActionCommands::UnRegisterCommandBindings(Toolkit->GetToolkitCommands());
 
 	// clear realtime viewport override
 	ConfigureRealTimeViewportsOverride(false);
@@ -651,7 +652,7 @@ void UModelingToolsEditorMode::CreateToolkit()
 
 void UModelingToolsEditorMode::OnToolStarted(UInteractiveToolManager* Manager, UInteractiveTool* Tool)
 {
-	FModelingToolActionCommands::UpdateToolCommandBinding(Tool, ToolCommandList, false);
+	FModelingToolActionCommands::UpdateToolCommandBinding(Tool, Toolkit->GetToolkitCommands(), false);
 	
 	if( FEngineAnalytics::IsAvailable() )
 	{
@@ -663,7 +664,7 @@ void UModelingToolsEditorMode::OnToolStarted(UInteractiveToolManager* Manager, U
 
 void UModelingToolsEditorMode::OnToolEnded(UInteractiveToolManager* Manager, UInteractiveTool* Tool)
 {
-	FModelingToolActionCommands::UpdateToolCommandBinding(Tool, ToolCommandList, true);
+	FModelingToolActionCommands::UpdateToolCommandBinding(Tool, Toolkit->GetToolkitCommands(), true);
 	
 	if( FEngineAnalytics::IsAvailable() )
 	{
