MeshModelingTools: Fixed tools responding to wasd keys while right-click flying after moving commands to toolkit command list. Fixed Enter not triggering complete in complete-style tools.

#rnx

diff --git a/Engine/Plugins/Experimental/ModelingToolsEditorMode/Source/ModelingToolsEditorMode/Private/ModelingToolsEditorMode.cpp b/Engine/Plugins/Experimental/ModelingToolsEditorMode/Source/ModelingToolsEditorMode/Private/ModelingToolsEditorMode.cpp
--- a/Engine/Plugins/Experimental/ModelingToolsEditorMode/Source/ModelingToolsEditorMode/Private/ModelingToolsEditorMode.cpp
+++ b/Engine/Plugins/Experimental/ModelingToolsEditorMode/Source/ModelingToolsEditorMode/Private/ModelingToolsEditorMode.cpp
@@ -705,6 +705,20 @@ void UModelingToolsEditorMode::BindCommands()
 		FIsActionButtonVisible::CreateLambda([this]() {return ToolsContext->CanCompleteActiveTool(); }),
 		EUIActionRepeatMode::RepeatDisabled
 		);
+
+	// These aren't activated by buttons but have default chords that bind the keypresses to the action.
+	CommandList->MapAction(
+		ToolManagerCommands.AcceptOrCompleteActiveTool,
+		FExecuteAction::CreateLambda([this]() {
+			const EToolShutdownType ShutdownType = ToolsContext->CanAcceptActiveTool() ? EToolShutdownType::Accept : EToolShutdownType::Completed;
+			ToolsContext->EndTool(ShutdownType);
+			}),
+		FCanExecuteAction::CreateLambda([this]() {
+				return ToolsContext->CanAcceptActiveTool() || ToolsContext->CanCompleteActiveTool();
+			}),
+		FGetActionCheckState(),
+		FIsActionButtonVisible(),
+		EUIActionRepeatMode::RepeatDisabled);
 	CommandList->MapAction(
 		ToolManagerCommands.CancelOrCompleteActiveTool,
 		FExecuteAction::CreateLambda([this]() {
@@ -714,9 +728,9 @@ void UModelingToolsEditorMode::BindCommands()
 		FCanExecuteAction::CreateLambda([this]() {
 				return ToolsContext->CanCompleteActiveTool() || ToolsContext->CanCancelActiveTool();
 			}),
-				FGetActionCheckState(),
-				FIsActionButtonVisible::CreateLambda([this]() { return ToolsContext->CanCompleteActiveTool() || ToolsContext->CanCancelActiveTool(); }),
-				EUIActionRepeatMode::RepeatDisabled);
+		FGetActionCheckState(),
+		FIsActionButtonVisible(),
+		EUIActionRepeatMode::RepeatDisabled);
 }
 
 void UModelingToolsEditorMode::ModelingModeShortcutRequested(EModelingModeActionCommands Command)
@@ -755,6 +769,20 @@ bool UModelingToolsEditorMode::GetPivotForOrbit(FVector& OutPivot) const
 	return false;
 }
 
+bool UModelingToolsEditorMode::InputKey(FEditorViewportClient* ViewportClient, FViewport* Viewport, FKey Key, EInputEvent Event)
+{
+	// TODO: This is a bit of a hack around the fact that when we fly around with right click + wasd,
+	// we still get the key presses passed to us. UEdMode currently does this check  for ToolCommandList,
+	// but we want our hotkeys to live in the toolkit command list so that we respond to them when the
+	// viewport is not focused.
+	// This should get removed when we have fixed wasd flying to capture its keys properly.
+	if (ToolsContext->ShouldIgnoreHotkeys())
+	{
+		return false;
+	}
+	return UEdMode::InputKey(ViewportClient, Viewport, Key, Event);
+}
+
 
 
 void UModelingToolsEditorMode::ConfigureRealTimeViewportsOverride(bool bEnable)
diff --git a/Engine/Plugins/Experimental/ModelingToolsEditorMode/Source/ModelingToolsEditorMode/Private/ModelingToolsManagerActions.cpp b/Engine/Plugins/Experimental/ModelingToolsEditorMode/Source/ModelingToolsEditorMode/Private/ModelingToolsManagerActions.cpp
--- a/Engine/Plugins/Experimental/ModelingToolsEditorMode/Source/ModelingToolsEditorMode/Private/ModelingToolsManagerActions.cpp
+++ b/Engine/Plugins/Experimental/ModelingToolsEditorMode/Source/ModelingToolsEditorMode/Private/ModelingToolsManagerActions.cpp
@@ -126,10 +126,12 @@ void FModelingToolsManagerCommands::RegisterCommands()
 	UI_COMMAND(BeginPolyModelTool_Outset, "Outset", "Start the PolyGroup Outset Tool", EUserInterfaceActionType::ToggleButton, FInputChord());
 	UI_COMMAND(BeginPolyModelTool_CutFaces, "Cut", "Start the PolyGroup Cut Faces Tool", EUserInterfaceActionType::ToggleButton, FInputChord());
 
-	// Note that passing a chord into one of these calls hooks the key press to the respective action. 
-	UI_COMMAND(AcceptActiveTool, "Accept", "Accept the active tool", EUserInterfaceActionType::Button, FInputChord(EKeys::Enter));
+	UI_COMMAND(AcceptActiveTool, "Accept", "Accept the active tool", EUserInterfaceActionType::Button, FInputChord());
 	UI_COMMAND(CancelActiveTool, "Cancel", "Cancel the active tool", EUserInterfaceActionType::Button, FInputChord());
-	UI_COMMAND(CompleteActiveTool, "Complete", "Complete the active tool", EUserInterfaceActionType::Button, FInputChord(EKeys::Enter));
+	UI_COMMAND(CompleteActiveTool, "Complete", "Complete the active tool", EUserInterfaceActionType::Button, FInputChord());
+
+	// Note that passing a chord into one of these calls hooks the key press to the respective action. 
+	UI_COMMAND(AcceptOrCompleteActiveTool, "Accept or Complete", "Accept or complete the active tool", EUserInterfaceActionType::Button, FInputChord(EKeys::Enter));
 	UI_COMMAND(CancelOrCompleteActiveTool, "Cancel or Complete", "Cancel or complete the active tool", EUserInterfaceActionType::Button, FInputChord(EKeys::Escape));
 }
 
diff --git a/Engine/Plugins/Experimental/ModelingToolsEditorMode/Source/ModelingToolsEditorMode/Public/ModelingToolsEditorMode.h b/Engine/Plugins/Experimental/ModelingToolsEditorMode/Source/ModelingToolsEditorMode/Public/ModelingToolsEditorMode.h
--- a/Engine/Plugins/Experimental/ModelingToolsEditorMode/Source/ModelingToolsEditorMode/Public/ModelingToolsEditorMode.h
+++ b/Engine/Plugins/Experimental/ModelingToolsEditorMode/Source/ModelingToolsEditorMode/Public/ModelingToolsEditorMode.h
@@ -42,6 +42,8 @@ public:
 
 	virtual bool GetPivotForOrbit(FVector& OutPivot) const override;
 
+	virtual bool InputKey(FEditorViewportClient* ViewportClient, FViewport* Viewport, FKey Key, EInputEvent Event) override;
+
 	/*
 	 * focus events
 	 */
diff --git a/Engine/Plugins/Experimental/ModelingToolsEditorMode/Source/ModelingToolsEditorMode/Public/ModelingToolsManagerActions.h b/Engine/Plugins/Experimental/ModelingToolsEditorMode/Source/ModelingToolsEditorMode/Public/ModelingToolsManagerActions.h
--- a/Engine/Plugins/Experimental/ModelingToolsEditorMode/Source/ModelingToolsEditorMode/Public/ModelingToolsManagerActions.h
+++ b/Engine/Plugins/Experimental/ModelingToolsEditorMode/Source/ModelingToolsEditorMode/Public/ModelingToolsManagerActions.h
@@ -122,6 +122,8 @@ public:
 	TSharedPtr<FUICommandInfo> AcceptActiveTool;
 	TSharedPtr<FUICommandInfo> CancelActiveTool;
 	TSharedPtr<FUICommandInfo> CompleteActiveTool;
+
+	TSharedPtr<FUICommandInfo> AcceptOrCompleteActiveTool;
 	TSharedPtr<FUICommandInfo> CancelOrCompleteActiveTool;
 
 	/**
