ModelingMode UX: change button labels to better fit new UX layout, re-organize tool categories, update tooltips. Move Accept/Cancel/Complete to bottom of details panel temporarily.
#rnx

diff --git a/Engine/Plugins/Experimental/ModelingToolsEditorMode/Source/ModelingToolsEditorMode/Private/ModelingToolsEditorModeToolkit.cpp b/Engine/Plugins/Experimental/ModelingToolsEditorMode/Source/ModelingToolsEditorMode/Private/ModelingToolsEditorModeToolkit.cpp
--- a/Engine/Plugins/Experimental/ModelingToolsEditorMode/Source/ModelingToolsEditorMode/Private/ModelingToolsEditorModeToolkit.cpp
+++ b/Engine/Plugins/Experimental/ModelingToolsEditorMode/Source/ModelingToolsEditorMode/Private/ModelingToolsEditorModeToolkit.cpp
@@ -26,11 +26,11 @@
 // if set to 1, then on mode initialization we include buttons for prototype modeling tools
 static TAutoConsoleVariable<int32> CVarEnablePrototypeModelingTools(
 	TEXT("modeling.EnablePrototypes"),
-	1,
+	0,
 	TEXT("Enable unsupported Experimental prototype Modeling Tools"));
 static TAutoConsoleVariable<int32> CVarEnablePolyModeling(
 	TEXT("modeling.EnablePolyModel"),
-	1,
+	0,
 	TEXT("Enable prototype PolyEdit tab"));
 
 
@@ -101,4 +101,9 @@
 
 			+ SVerticalBox::Slot().AutoHeight().HAlign(HAlign_Fill).Padding(5)
 				[
+					MakeAcceptCancelPanel()->AsShared()
+				]
+
+			+ SVerticalBox::Slot().AutoHeight().HAlign(HAlign_Fill).Padding(5)
+				[
 					MakeAssetConfigPanel()->AsShared()
@@ -117,6 +122,42 @@
 
 
 
+
+TSharedPtr<SWidget> FModelingToolsEditorModeToolkit::MakeAcceptCancelPanel()
+{
+
+	TSharedPtr<FUICommandList> CommandList = GetToolkitCommands();
+	FName ToolbarCustomizationName = GetEditorMode() ? GetEditorMode()->GetModeInfo().ToolbarCustomizationName : GetScriptableEditorMode()->GetModeInfo().ToolbarCustomizationName;
+	FUniformToolBarBuilder RowBuilder(CommandList, FMultiBoxCustomization(ToolbarCustomizationName));
+	RowBuilder.SetStyle(&FEditorStyle::Get(), "PaletteToolBar");
+
+	const FModelingToolsManagerCommands& Commands = FModelingToolsManagerCommands::Get();
+	RowBuilder.AddToolBarButton(Commands.AcceptActiveTool);
+	RowBuilder.AddToolBarButton(Commands.CancelActiveTool);
+	RowBuilder.AddToolBarButton(Commands.CompleteActiveTool);
+
+	TSharedPtr<SVerticalBox> Content = SNew(SVerticalBox)
+	+ SVerticalBox::Slot().AutoHeight().HAlign(HAlign_Right).Padding(0, 4, 0, 10)
+		[
+			RowBuilder.MakeWidget()
+		];
+
+	TSharedPtr<SExpandableArea> AcceptCancelPanel = SNew(SExpandableArea)
+		.HeaderPadding(FMargin(2.0f))
+		.Padding(FMargin(2.f))
+		.BorderImage(FEditorStyle::Get().GetBrush("DetailsView.CategoryTop"))
+		.BorderBackgroundColor(FLinearColor(.6, .6, .6, 1.0f))
+		.BodyBorderBackgroundColor(FLinearColor::Transparent)
+		.AreaTitleFont(FEditorStyle::Get().GetFontStyle("EditorModesPanel.CategoryFontStyle"))
+		.BodyContent()
+		[
+			Content->AsShared()
+		];
+
+	return AcceptCancelPanel;
+}
+
+
 TSharedPtr<SWidget> FModelingToolsEditorModeToolkit::MakeAssetConfigPanel()
 {
 	AssetLocationModes.Add(MakeShared<FString>(TEXT("AutoGen Folder (World-Relative)")));
@@ -265,29 +306,34 @@
 	return NSLOCTEXT("ModelingToolsEditorModeToolkit", "DisplayName", "ModelingToolsEditorMode Tool");
 }
 
-static const FName PrimitiveTabName(TEXT("Primitives"));
+static const FName PrimitiveTabName(TEXT("Shapes"));
 static const FName CreateTabName(TEXT("Create"));
 static const FName AttributesTabName(TEXT("Attributes"));
-static const FName SculptTabName(TEXT("Sculpt"));
-static const FName TrianglesTabName(TEXT("Triangles"));
-static const FName PolyModelingTabName(TEXT("PolyModeling"));
+static const FName TriModelingTabName(TEXT("TriModel"));
+static const FName PolyModelingTabName(TEXT("PolyModel"));
+static const FName MeshProcessingTabName(TEXT("MeshOps"));
 static const FName UVTabName(TEXT("UVs"));
 static const FName TransformTabName(TEXT("Transform"));
-static const FName ModifyTabName(TEXT("Modify"));
+static const FName DeformTabName(TEXT("Deform"));
 static const FName VolumesTabName(TEXT("Volumes"));
 static const FName PrototypesTabName(TEXT("Prototypes"));
 static const FName HairTabName(TEXT("Hair"));
 static const FName PolyEditTabName(TEXT("PolyEdit"));
+static const FName VoxToolsTabName(TEXT("VoxOps"));
 
 
-const TArray<FName> FModelingToolsEditorModeToolkit::PaletteNames_Standard = { PrimitiveTabName, CreateTabName, PolyModelingTabName, TrianglesTabName, ModifyTabName, TransformTabName, AttributesTabName, UVTabName, VolumesTabName, HairTabName };
-const TArray<FName> FModelingToolsEditorModeToolkit::PaletteNames_Experimental = { PrimitiveTabName, CreateTabName, PolyModelingTabName, TrianglesTabName, ModifyTabName, TransformTabName, AttributesTabName, UVTabName, VolumesTabName, HairTabName, PrototypesTabName };
+const TArray<FName> FModelingToolsEditorModeToolkit::PaletteNames_Standard = { PrimitiveTabName, CreateTabName, PolyModelingTabName, TriModelingTabName, DeformTabName, TransformTabName, MeshProcessingTabName, VoxToolsTabName, AttributesTabName, UVTabName, VolumesTabName, HairTabName };
 
 
 void FModelingToolsEditorModeToolkit::GetToolPaletteNames(TArray<FName>& PaletteNames) const
 {
+	PaletteNames = PaletteNames_Standard;
+
 	bool bEnablePrototypes = (CVarEnablePrototypeModelingTools.GetValueOnGameThread() > 0);
-	PaletteNames = (bEnablePrototypes) ? PaletteNames_Experimental : PaletteNames_Standard;
+	if (bEnablePrototypes)
+	{
+		PaletteNames.Add(PrototypesTabName);
+	}
 
 	bool bEnablePolyModel = (CVarEnablePolyModeling.GetValueOnGameThread() > 0);
 	if (bEnablePolyModel)
@@ -323,97 +369,91 @@
 	if (PaletteIndex == PrimitiveTabName)
 	{
 		ToolbarBuilder.AddToolBarButton(Commands.BeginAddBoxPrimitiveTool);
+		ToolbarBuilder.AddToolBarButton(Commands.BeginAddSpherePrimitiveTool);
+		ToolbarBuilder.AddToolBarButton(Commands.BeginAddSphericalBoxPrimitiveTool);
 		ToolbarBuilder.AddToolBarButton(Commands.BeginAddCylinderPrimitiveTool);
 		ToolbarBuilder.AddToolBarButton(Commands.BeginAddConePrimitiveTool);
+		ToolbarBuilder.AddToolBarButton(Commands.BeginAddTorusPrimitiveTool);
 		ToolbarBuilder.AddToolBarButton(Commands.BeginAddArrowPrimitiveTool);
 		ToolbarBuilder.AddToolBarButton(Commands.BeginAddRectanglePrimitiveTool);
 		ToolbarBuilder.AddToolBarButton(Commands.BeginAddRoundedRectanglePrimitiveTool);
 		ToolbarBuilder.AddToolBarButton(Commands.BeginAddDiscPrimitiveTool);
 		ToolbarBuilder.AddToolBarButton(Commands.BeginAddPuncturedDiscPrimitiveTool);
-		ToolbarBuilder.AddToolBarButton(Commands.BeginAddTorusPrimitiveTool);
-		ToolbarBuilder.AddToolBarButton(Commands.BeginAddSpherePrimitiveTool);
-		ToolbarBuilder.AddToolBarButton(Commands.BeginAddSphericalBoxPrimitiveTool);
 	}
 	else if (PaletteIndex == CreateTabName)
 	{
-		ToolbarBuilder.AddSeparator();
+		ToolbarBuilder.AddToolBarButton(Commands.BeginCombineMeshesTool);
+		ToolbarBuilder.AddToolBarButton(Commands.BeginDuplicateMeshesTool);
 		ToolbarBuilder.AddToolBarButton(Commands.BeginDrawPolygonTool);
 		ToolbarBuilder.AddToolBarButton(Commands.BeginDrawPolyPathTool);
 		ToolbarBuilder.AddToolBarButton(Commands.BeginDrawAndRevolveTool);
 		ToolbarBuilder.AddToolBarButton(Commands.BeginRevolveBoundaryTool);
-		ToolbarBuilder.AddSeparator();
-		ToolbarBuilder.AddToolBarButton(Commands.BeginMeshBooleanTool);
-		ToolbarBuilder.AddToolBarButton(Commands.BeginMeshTrimTool);
-		ToolbarBuilder.AddToolBarButton(Commands.BeginSelfUnionTool);
-		ToolbarBuilder.AddSeparator();
-		ToolbarBuilder.AddToolBarButton(Commands.BeginVoxelSolidifyTool);
-		ToolbarBuilder.AddToolBarButton(Commands.BeginVoxelBlendTool);
-		ToolbarBuilder.AddToolBarButton(Commands.BeginVoxelMorphologyTool);
-		ToolbarBuilder.AddToolBarButton(Commands.BeginVoxelBooleanTool);
-		ToolbarBuilder.AddToolBarButton(Commands.BeginVoxelMergeTool);
 	}
 	else if (PaletteIndex == TransformTabName)
 	{
-		ToolbarBuilder.AddToolBarButton(Commands.BeginEditPivotTool);
-		ToolbarBuilder.AddToolBarButton(Commands.BeginBakeTransformTool);
-		ToolbarBuilder.AddSeparator();
-		ToolbarBuilder.AddToolBarButton(Commands.BeginCombineMeshesTool);
-		ToolbarBuilder.AddToolBarButton(Commands.BeginDuplicateMeshesTool);
-		ToolbarBuilder.AddSeparator();
-		ToolbarBuilder.AddToolBarButton(Commands.BeginGenerateStaticMeshLODAssetTool);
-		ToolbarBuilder.AddSeparator();
 		ToolbarBuilder.AddToolBarButton(Commands.BeginTransformMeshesTool);
 		ToolbarBuilder.AddToolBarButton(Commands.BeginAlignObjectsTool);
+		ToolbarBuilder.AddToolBarButton(Commands.BeginEditPivotTool);
+		ToolbarBuilder.AddToolBarButton(Commands.BeginBakeTransformTool);
 	}
-	else if (PaletteIndex == ModifyTabName)
+	else if (PaletteIndex == DeformTabName)
 	{
 		ToolbarBuilder.AddToolBarButton(Commands.BeginSculptMeshTool);
 		ToolbarBuilder.AddToolBarButton(Commands.BeginRemeshSculptMeshTool);
-		ToolbarBuilder.AddSeparator();
-		ToolbarBuilder.AddToolBarButton(Commands.BeginPlaneCutTool);
-		ToolbarBuilder.AddToolBarButton(Commands.BeginMirrorTool);
-		ToolbarBuilder.AddToolBarButton(Commands.BeginPolygonCutTool);
 		ToolbarBuilder.AddToolBarButton(Commands.BeginSmoothMeshTool);
 		ToolbarBuilder.AddToolBarButton(Commands.BeginOffsetMeshTool);
-		ToolbarBuilder.AddToolBarButton(Commands.BeginDisplaceMeshTool);
 		ToolbarBuilder.AddToolBarButton(Commands.BeginMeshSpaceDeformerTool);
 		ToolbarBuilder.AddToolBarButton(Commands.BeginLatticeDeformerTool);
-		ToolbarBuilder.AddToolBarButton(Commands.BeginRemeshMeshTool);
+		ToolbarBuilder.AddToolBarButton(Commands.BeginDisplaceMeshTool);
 	}
-	else if (PaletteIndex == TrianglesTabName)
+	else if (PaletteIndex == MeshProcessingTabName)
 	{
-		ToolbarBuilder.AddToolBarButton(Commands.BeginMeshSelectionTool);
-		ToolbarBuilder.AddSeparator();
-		ToolbarBuilder.AddToolBarButton(Commands.BeginTriEditTool);
 		ToolbarBuilder.AddToolBarButton(Commands.BeginSimplifyMeshTool);
 		ToolbarBuilder.AddToolBarButton(Commands.BeginRemeshMeshTool);
-		ToolbarBuilder.AddToolBarButton(Commands.BeginProjectToTargetTool);
-		ToolbarBuilder.AddSeparator();
-		ToolbarBuilder.AddToolBarButton(Commands.BeginEditMeshMaterialsTool);
 		ToolbarBuilder.AddToolBarButton(Commands.BeginWeldEdgesTool);
 		ToolbarBuilder.AddToolBarButton(Commands.BeginRemoveOccludedTrianglesTool);
+		ToolbarBuilder.AddToolBarButton(Commands.BeginSelfUnionTool);
+		ToolbarBuilder.AddToolBarButton(Commands.BeginProjectToTargetTool);
+		ToolbarBuilder.AddToolBarButton(Commands.BeginGenerateStaticMeshLODAssetTool);
+	}
+	else if (PaletteIndex == VoxToolsTabName)
+	{
+		ToolbarBuilder.AddToolBarButton(Commands.BeginVoxelSolidifyTool);
+		ToolbarBuilder.AddToolBarButton(Commands.BeginVoxelBlendTool);
+		ToolbarBuilder.AddToolBarButton(Commands.BeginVoxelMorphologyTool);
+		ToolbarBuilder.AddToolBarButton(Commands.BeginVoxelBooleanTool);
+		ToolbarBuilder.AddToolBarButton(Commands.BeginVoxelMergeTool);
+	}
+	else if (PaletteIndex == TriModelingTabName)
+	{
+		ToolbarBuilder.AddToolBarButton(Commands.BeginMeshSelectionTool);
+		ToolbarBuilder.AddToolBarButton(Commands.BeginTriEditTool);
 		ToolbarBuilder.AddToolBarButton(Commands.BeginHoleFillTool);
+		ToolbarBuilder.AddToolBarButton(Commands.BeginPlaneCutTool);
+		ToolbarBuilder.AddToolBarButton(Commands.BeginMirrorTool);
+		ToolbarBuilder.AddToolBarButton(Commands.BeginPolygonCutTool);
+		ToolbarBuilder.AddToolBarButton(Commands.BeginMeshTrimTool);
 	}
 	else if (PaletteIndex == PolyModelingTabName)
 	{
 		ToolbarBuilder.AddToolBarButton(Commands.BeginPolyEditTool);
 		ToolbarBuilder.AddToolBarButton(Commands.BeginPolyDeformTool);
-		ToolbarBuilder.AddSeparator();
 		ToolbarBuilder.AddToolBarButton(Commands.BeginGroupEdgeInsertionTool);
 		ToolbarBuilder.AddToolBarButton(Commands.BeginEdgeLoopInsertionTool);
-		ToolbarBuilder.AddSeparator();
-		ToolbarBuilder.AddToolBarButton(Commands.BeginDrawPolygonTool);
 		ToolbarBuilder.AddToolBarButton(Commands.BeginMeshBooleanTool);
+		ToolbarBuilder.AddToolBarButton(Commands.BeginSubdividePolyTool);
 	}
 	else if (PaletteIndex == AttributesTabName)
 	{
-		ToolbarBuilder.AddToolBarButton(Commands.BeginPolyGroupsTool);
-		ToolbarBuilder.AddToolBarButton(Commands.BeginAttributeEditorTool);
 		ToolbarBuilder.AddToolBarButton(Commands.BeginMeshInspectorTool);
-		ToolbarBuilder.AddToolBarButton(Commands.BeginMeshGroupPaintTool);
 		ToolbarBuilder.AddToolBarButton(Commands.BeginEditNormalsTool);
 		ToolbarBuilder.AddToolBarButton(Commands.BeginEditTangentsTool);
+		ToolbarBuilder.AddToolBarButton(Commands.BeginAttributeEditorTool);
+		ToolbarBuilder.AddToolBarButton(Commands.BeginPolyGroupsTool);
+		ToolbarBuilder.AddToolBarButton(Commands.BeginMeshGroupPaintTool);
 		ToolbarBuilder.AddToolBarButton(Commands.BeginMeshAttributePaintTool);
+		ToolbarBuilder.AddToolBarButton(Commands.BeginEditMeshMaterialsTool);
+		ToolbarBuilder.AddToolBarButton(Commands.BeginBakeMeshAttributeMapsTool);
 	}
 	else if (PaletteIndex == UVTabName)
 	{
@@ -423,8 +463,6 @@
 		ToolbarBuilder.AddToolBarButton(Commands.BeginUVSeamEditTool);
 		ToolbarBuilder.AddToolBarButton(Commands.BeginTransformUVIslandsTool);
 		ToolbarBuilder.AddToolBarButton(Commands.BeginUVLayoutTool);
-		ToolbarBuilder.AddSeparator();
-		ToolbarBuilder.AddToolBarButton(Commands.BeginBakeMeshAttributeMapsTool);
 	}
 	else if (PaletteIndex == VolumesTabName)
 	{
@@ -448,7 +486,6 @@
 	{
 		ToolbarBuilder.AddToolBarButton(Commands.BeginAddPatchTool);
 		ToolbarBuilder.AddToolBarButton(Commands.BeginShapeSprayTool);
-		ToolbarBuilder.AddToolBarButton(Commands.BeginSubdividePolyTool);
 	}
 	else if (PaletteIndex == PolyEditTabName)
 	{
@@ -471,13 +508,6 @@
 		ToolbarBuilder.AddToolBarButton(Commands.BeginSubdividePolyTool);
 		ToolbarBuilder.AddToolBarButton(Commands.BeginPolyEditTool);
 	}
-
-	ToolbarBuilder.AddSeparator();
-
-	ToolbarBuilder.AddToolBarButton(Commands.AcceptActiveTool);
-	ToolbarBuilder.AddToolBarButton(Commands.CancelActiveTool);
-	ToolbarBuilder.AddToolBarButton(Commands.CompleteActiveTool);
-
 }
 
 void FModelingToolsEditorModeToolkit::OnToolPaletteChanged(FName PaletteName) 
diff --git a/Engine/Plugins/Experimental/ModelingToolsEditorMode/Source/ModelingToolsEditorMode/Private/ModelingToolsManagerActions.cpp b/Engine/Plugins/Experimental/ModelingToolsEditorMode/Source/ModelingToolsEditorMode/Private/ModelingToolsManagerActions.cpp
--- a/Engine/Plugins/Experimental/ModelingToolsEditorMode/Source/ModelingToolsEditorMode/Private/ModelingToolsManagerActions.cpp
+++ b/Engine/Plugins/Experimental/ModelingToolsEditorMode/Source/ModelingToolsEditorMode/Private/ModelingToolsManagerActions.cpp
@@ -22,113 +22,113 @@ FModelingToolsManagerCommands::FModelingToolsManagerCommands() :
 
 void FModelingToolsManagerCommands::RegisterCommands()
 {
-	UI_COMMAND(BeginAddBoxPrimitiveTool, "Box", "Add Boxes", EUserInterfaceActionType::ToggleButton, FInputChord());
-	UI_COMMAND(BeginAddCylinderPrimitiveTool, "Cylinder", "Add Cylinders", EUserInterfaceActionType::ToggleButton, FInputChord());
-	UI_COMMAND(BeginAddConePrimitiveTool, "Cone", "Add Cones", EUserInterfaceActionType::ToggleButton, FInputChord());
-	UI_COMMAND(BeginAddArrowPrimitiveTool, "Arrow", "Add Arrows", EUserInterfaceActionType::ToggleButton, FInputChord());
-	UI_COMMAND(BeginAddRectanglePrimitiveTool, "Rectangle", "Add Rectangles", EUserInterfaceActionType::ToggleButton, FInputChord());
-	UI_COMMAND(BeginAddRoundedRectanglePrimitiveTool, "Rounded Rectangle", "Add Rounded Rectangles", EUserInterfaceActionType::ToggleButton, FInputChord());
-	UI_COMMAND(BeginAddDiscPrimitiveTool, "Disc", "Add Discs", EUserInterfaceActionType::ToggleButton, FInputChord());
-	UI_COMMAND(BeginAddPuncturedDiscPrimitiveTool, "Punctured Disc", "Add PuncturedDiscs", EUserInterfaceActionType::ToggleButton, FInputChord());
-	UI_COMMAND(BeginAddTorusPrimitiveTool, "Torus", "Add Torii", EUserInterfaceActionType::ToggleButton, FInputChord());
-	UI_COMMAND(BeginAddSphericalBoxPrimitiveTool, "Spherical Box", "Add Box Parametrized Spheres", EUserInterfaceActionType::ToggleButton, FInputChord());
-	UI_COMMAND(BeginAddSpherePrimitiveTool, "Sphere", "Add Lat/Long Parametrized Spheres", EUserInterfaceActionType::ToggleButton, FInputChord());
-
-	UI_COMMAND(BeginAddPatchTool, "Patch", "Start the Add Patch Tool", EUserInterfaceActionType::ToggleButton, FInputChord());
-	UI_COMMAND(BeginRevolveBoundaryTool, "CutRevolve", "Start the CutRevolve Tool, which revolves a section curve created by cutting the selected Mesh", EUserInterfaceActionType::ToggleButton, FInputChord());
-	UI_COMMAND(BeginDrawPolygonTool, "Polygon", "Start the Draw Polygon Tool", EUserInterfaceActionType::ToggleButton, FInputChord());
-	UI_COMMAND(BeginDrawPolyPathTool, "PolyPath", "Start the Draw PolyPath Tool", EUserInterfaceActionType::ToggleButton, FInputChord());
-	UI_COMMAND(BeginDrawAndRevolveTool, "PolyRevolve", "Start the Draw-and-Revolve Tool", EUserInterfaceActionType::ToggleButton, FInputChord());
-	UI_COMMAND(BeginShapeSprayTool, "ShapeSpray", "Start the Shape Spray Tool", EUserInterfaceActionType::ToggleButton, FInputChord());
-
-	UI_COMMAND(BeginSculptMeshTool, "Sculpt", "Start the Sculpt Mesh Tool", EUserInterfaceActionType::ToggleButton, FInputChord());
-	UI_COMMAND(BeginPolyEditTool, "PolyEdit", "Start the PolyEdit Tool", EUserInterfaceActionType::ToggleButton, FInputChord());
-	UI_COMMAND(BeginGroupEdgeInsertionTool, "EdgeInsert", "Start the Group Edge Insertion Tool", EUserInterfaceActionType::ToggleButton, FInputChord());
-	UI_COMMAND(BeginEdgeLoopInsertionTool, "LoopInsert", "Start the Edge Loop Insertion Tool", EUserInterfaceActionType::ToggleButton, FInputChord());
-	UI_COMMAND(BeginTriEditTool, "TriEdit", "Start the Triangle Edit Tool", EUserInterfaceActionType::ToggleButton, FInputChord());
-	UI_COMMAND(BeginPolyDeformTool, "PolyDeform", "Start the PolyDeform Tool", EUserInterfaceActionType::ToggleButton, FInputChord());
-	UI_COMMAND(BeginSmoothMeshTool, "Smooth", "Start the Smooth Mesh Tool", EUserInterfaceActionType::ToggleButton, FInputChord());
-	UI_COMMAND(BeginOffsetMeshTool, "Offset", "Start the Offset Mesh Tool", EUserInterfaceActionType::ToggleButton, FInputChord());
-	UI_COMMAND(BeginDisplaceMeshTool, "Displace", "Start the Displace Mesh Tool", EUserInterfaceActionType::ToggleButton, FInputChord());
-	UI_COMMAND(BeginMeshSpaceDeformerTool, "SpaceDeform", "Start the Mesh Space Deformer Tool", EUserInterfaceActionType::ToggleButton, FInputChord());
-	UI_COMMAND(BeginTransformMeshesTool, "Transform", "Start the Transform Tool", EUserInterfaceActionType::ToggleButton, FInputChord());
-	UI_COMMAND(BeginEditPivotTool, "Edit Pivot", "Start the Edit Pivot Tool", EUserInterfaceActionType::ToggleButton, FInputChord());
-	UI_COMMAND(BeginBakeTransformTool, "Bake XForm", "Start the Bake Transform Tool", EUserInterfaceActionType::ToggleButton, FInputChord());
-	UI_COMMAND(BeginCombineMeshesTool, "Combine", "Combine the selected Mesh Assets into a new Asset", EUserInterfaceActionType::ToggleButton, FInputChord());
-	UI_COMMAND(BeginDuplicateMeshesTool, "Duplicate", "Duplicate the selected Mesh Asset into a new Asset", EUserInterfaceActionType::ToggleButton, FInputChord());
-	UI_COMMAND(BeginAlignObjectsTool, "Align", "Start the Align Objects Tool", EUserInterfaceActionType::ToggleButton, FInputChord());
-
-	UI_COMMAND(BeginRemeshSculptMeshTool, "DynaSculpt", "Start the Liquid Sculpt Mesh Tool", EUserInterfaceActionType::ToggleButton, FInputChord());
-	UI_COMMAND(BeginRemeshMeshTool, "Remesh", "Start the Remesh Tool", EUserInterfaceActionType::ToggleButton, FInputChord());
-	UI_COMMAND(BeginProjectToTargetTool, "MapToMesh", "Start the Remesh To Target Tool", EUserInterfaceActionType::ToggleButton, FInputChord());
-	UI_COMMAND(BeginSimplifyMeshTool, "Simplify", "Start the Simplify Mesh Tool", EUserInterfaceActionType::ToggleButton, FInputChord());
-	UI_COMMAND(BeginEditNormalsTool, "Normals", "Start the Edit Normals Tool", EUserInterfaceActionType::ToggleButton, FInputChord());
-	UI_COMMAND(BeginEditTangentsTool, "Tangents", "Start the Edit Tangents Tool", EUserInterfaceActionType::ToggleButton, FInputChord());
-	UI_COMMAND(BeginRemoveOccludedTrianglesTool, "Jacketing", "Start the Jacketing Tool", EUserInterfaceActionType::ToggleButton, FInputChord());
-	UI_COMMAND(BeginUVProjectionTool, "UVProjection", "Start the UV Projection Tool", EUserInterfaceActionType::ToggleButton, FInputChord());
-	UI_COMMAND(BeginUVLayoutTool, "UVLayout", "Start the UV Layout Tool", EUserInterfaceActionType::ToggleButton, FInputChord());
-	UI_COMMAND(BeginVoxelMergeTool, "VoxMerge", "Start the Voxel Merge Tool", EUserInterfaceActionType::ToggleButton, FInputChord());
-	UI_COMMAND(BeginVoxelBooleanTool, "VoxBoolean", "Start the Voxel Boolean Tool", EUserInterfaceActionType::ToggleButton, FInputChord());
-	UI_COMMAND(BeginVoxelSolidifyTool, "VoxWrap", "Start the Voxel Shell Tool", EUserInterfaceActionType::ToggleButton, FInputChord());
-	UI_COMMAND(BeginVoxelBlendTool, "VoxBlend", "Start the Voxel Blend Tool", EUserInterfaceActionType::ToggleButton, FInputChord());
-	UI_COMMAND(BeginVoxelMorphologyTool, "VoxMorphology", "Start the Voxel Morphology Tool", EUserInterfaceActionType::ToggleButton, FInputChord());
-	UI_COMMAND(BeginSelfUnionTool, "SelfUnion", "Start the Mesh-Based Self-Union Tool, which will resolve self-intersections", EUserInterfaceActionType::ToggleButton, FInputChord());
-	UI_COMMAND(BeginMeshBooleanTool, "Mesh Boolean", "Start the Mesh-Based Boolean Tool", EUserInterfaceActionType::ToggleButton, FInputChord());
-	UI_COMMAND(BeginMeshTrimTool, "Mesh Trim", "Start the Mesh-Based Trim Tool", EUserInterfaceActionType::ToggleButton, FInputChord());
-	UI_COMMAND(BeginBspConversionTool, "BSP Conversion", "Start the BSP-to-StaticMesh Conversion Tool", EUserInterfaceActionType::ToggleButton, FInputChord());
-	UI_COMMAND(BeginMeshToVolumeTool, "Mesh2Vol", "Start the Mesh to Volume Conversion Tool", EUserInterfaceActionType::ToggleButton, FInputChord());
-	UI_COMMAND(BeginVolumeToMeshTool, "Vol2Mesh", "Start the Volume to Mesh Conversion Tool", EUserInterfaceActionType::ToggleButton, FInputChord());
-	UI_COMMAND(BeginPlaneCutTool, "PlaneCut", "Start the Plane Cut Tool", EUserInterfaceActionType::ToggleButton, FInputChord());
-	UI_COMMAND(BeginMirrorTool, "Mirror", "Start the Mirror Tool", EUserInterfaceActionType::ToggleButton, FInputChord());
-	UI_COMMAND(BeginHoleFillTool, "HoleFill", "Start the Hole Fill Tool", EUserInterfaceActionType::ToggleButton, FInputChord());
-	UI_COMMAND(BeginLatticeDeformerTool, "LatticeDeformer", "Start the Lattice Deformer Tool", EUserInterfaceActionType::ToggleButton, FInputChord());
-	UI_COMMAND(BeginPolygonCutTool, "PolyCut", "Start the Polygon Cut Tool", EUserInterfaceActionType::ToggleButton, FInputChord());
-	UI_COMMAND(BeginGlobalUVGenerateTool, "AutoUnwrap", "Start the Global UV Unwrap Tool", EUserInterfaceActionType::ToggleButton, FInputChord());
-	UI_COMMAND(BeginGroupUVGenerateTool, "GroupUnwrap", "Stat the PolyGroup UV Unwrap Tool", EUserInterfaceActionType::ToggleButton, FInputChord());
-	UI_COMMAND(BeginMeshSelectionTool, "TriSelect", "Start the Triangle Selection Tool", EUserInterfaceActionType::ToggleButton, FInputChord());
-
-	UI_COMMAND(BeginPhysicsInspectorTool, "PhysInspect", "Inspect Physics Geometry for selected Meshes", EUserInterfaceActionType::ToggleButton, FInputChord());
-	UI_COMMAND(BeginSetCollisionGeometryTool, "MeshToCollision", "Convert selected Meshes to Collision Geometry", EUserInterfaceActionType::ToggleButton, FInputChord());
-	UI_COMMAND(BeginEditCollisionGeometryTool, "EditPhys", "Edit Collision Geometry for selected Mesh", EUserInterfaceActionType::ToggleButton, FInputChord());
-	UI_COMMAND(BeginExtractCollisionGeometryTool, "CollisionToMesh", "Convert Collision Geometry to a Mesh", EUserInterfaceActionType::ToggleButton, FInputChord());
-
-	UI_COMMAND(BeginMeshInspectorTool, "Inspector", "Start the Mesh Inspector Tool", EUserInterfaceActionType::ToggleButton, FInputChord());
-	UI_COMMAND(BeginWeldEdgesTool, "Weld Edges", "Start the Weld Edges Tool", EUserInterfaceActionType::ToggleButton, FInputChord());
-	UI_COMMAND(BeginPolyGroupsTool, "PolyGroups", "Start the PolyGroups Tool", EUserInterfaceActionType::ToggleButton, FInputChord());
-	UI_COMMAND(BeginEditMeshMaterialsTool, "Edit Mats", "Start the Material Editing Tool", EUserInterfaceActionType::ToggleButton, FInputChord());
-	UI_COMMAND(BeginTransformUVIslandsTool, "Transform UVs", "Start the UV Island Transformation Tool", EUserInterfaceActionType::ToggleButton, FInputChord());
-	UI_COMMAND(BeginMeshAttributePaintTool, "Paint Maps", "Start the Paint Maps Tool", EUserInterfaceActionType::ToggleButton, FInputChord());
-	UI_COMMAND(BeginAttributeEditorTool, "Edit Attribs", "Start the Attribute Editor Tool", EUserInterfaceActionType::ToggleButton, FInputChord());
-	UI_COMMAND(BeginBakeMeshAttributeMapsTool, "Bake Maps", "Start the Map Baking Tool", EUserInterfaceActionType::ToggleButton, FInputChord());
-	UI_COMMAND(BeginUVSeamEditTool, "UVSeamEdit", "Start the UV Seam Editing Tool", EUserInterfaceActionType::ToggleButton, FInputChord());
-
-	UI_COMMAND(BeginGroomToMeshTool, "HairHelmet", "Start the Hair Helmet Tool", EUserInterfaceActionType::ToggleButton, FInputChord());
-	UI_COMMAND(BeginGenerateLODMeshesTool, "GenLODs", "Start the Generate LOD Meshes Tool", EUserInterfaceActionType::ToggleButton, FInputChord());
-	UI_COMMAND(BeginGroomCardsToMeshTool, "CardsToMesh", "Start the Hair Cards to Mesh Tool", EUserInterfaceActionType::ToggleButton, FInputChord());
-	UI_COMMAND(BeginGroomMeshToCardsTool, "MeshToCards", "Start the Mesh to Hair Cards Tool", EUserInterfaceActionType::ToggleButton, FInputChord());
-	UI_COMMAND(BeginGroomCardsEditorTool, "CardsEditor", "Start the Cards Editor Tool", EUserInterfaceActionType::ToggleButton, FInputChord());
-	UI_COMMAND(BeginSubdividePolyTool, "Subdivide", "Start the SubdividePoly Tool", EUserInterfaceActionType::ToggleButton, FInputChord());
-
-	UI_COMMAND(BeginGenerateStaticMeshLODAssetTool, "AutoLOD", "Start the Generate Static Mesh LOD Asset Tool", EUserInterfaceActionType::ToggleButton, FInputChord());
-	UI_COMMAND(BeginMeshGroupPaintTool, "GroupPaint", "Start the Mesh PolyGroup Paint Tool", EUserInterfaceActionType::ToggleButton, FInputChord());
-
-
-	UI_COMMAND(BeginPolyModelTool_FaceSelect, "Faces", "Start the PolyGroup Face Selection Tool", EUserInterfaceActionType::ToggleButton, FInputChord());
-	UI_COMMAND(BeginPolyModelTool_EdgeSelect, "Edges", "Start the PolyGroup Edge Selection Tool", EUserInterfaceActionType::ToggleButton, FInputChord());
-	UI_COMMAND(BeginPolyModelTool_VertexSelect, "Verts", "Start the PolyGroup Vertex Selection Tool", EUserInterfaceActionType::ToggleButton, FInputChord());
-	UI_COMMAND(BeginPolyModelTool_AllSelect, "Any", "Start the PolyGroup Face/Edge/Vertex Selection Tool", EUserInterfaceActionType::ToggleButton, FInputChord());
-	UI_COMMAND(BeginPolyModelTool_LoopSelect, "Loops", "Start the PolyGroup Loop Selection Tool", EUserInterfaceActionType::ToggleButton, FInputChord());
-	UI_COMMAND(BeginPolyModelTool_RingSelect, "Rings", "Start the PolyGroup Ring Selection Tool", EUserInterfaceActionType::ToggleButton, FInputChord());
-
-	UI_COMMAND(BeginPolyModelTool_Extrude, "Extrude", "Start the PolyGroup Extrude Tool", EUserInterfaceActionType::ToggleButton, FInputChord());
-	UI_COMMAND(BeginPolyModelTool_Offset, "Offset", "Start the PolyGroup Offset Tool", EUserInterfaceActionType::ToggleButton, FInputChord());
-	UI_COMMAND(BeginPolyModelTool_Inset, "Inset", "Start the PolyGroup Inset Tool", EUserInterfaceActionType::ToggleButton, FInputChord());
-	UI_COMMAND(BeginPolyModelTool_Outset, "Outset", "Start the PolyGroup Outset Tool", EUserInterfaceActionType::ToggleButton, FInputChord());
-	UI_COMMAND(BeginPolyModelTool_CutFaces, "Cut", "Start the PolyGroup Cut Faces Tool", EUserInterfaceActionType::ToggleButton, FInputChord());
+	UI_COMMAND(BeginAddBoxPrimitiveTool, "Box", "Create new Box StaticMesh Assets", EUserInterfaceActionType::ToggleButton, FInputChord());
+	UI_COMMAND(BeginAddCylinderPrimitiveTool, "Cyl", "Create new Cylinder StaticMesh Assets", EUserInterfaceActionType::ToggleButton, FInputChord());
+	UI_COMMAND(BeginAddConePrimitiveTool, "Cone", "Create new Cone StaticMesh Assets", EUserInterfaceActionType::ToggleButton, FInputChord());
+	UI_COMMAND(BeginAddArrowPrimitiveTool, "Arrow", "Create new Arrow StaticMesh Assets", EUserInterfaceActionType::ToggleButton, FInputChord());
+	UI_COMMAND(BeginAddRectanglePrimitiveTool, "Rect", "Create new Rectangle StaticMesh Assets", EUserInterfaceActionType::ToggleButton, FInputChord());
+	UI_COMMAND(BeginAddRoundedRectanglePrimitiveTool, "RndRct", "Create new Rounded Rectangle StaticMesh Assets", EUserInterfaceActionType::ToggleButton, FInputChord());
+	UI_COMMAND(BeginAddDiscPrimitiveTool, "Disc", "Create new Disc StaticMesh Assets", EUserInterfaceActionType::ToggleButton, FInputChord());
+	UI_COMMAND(BeginAddPuncturedDiscPrimitiveTool, "Circle", "Create new Punctured Disc StaticMesh Assets", EUserInterfaceActionType::ToggleButton, FInputChord());
+	UI_COMMAND(BeginAddTorusPrimitiveTool, "Torus", "Create new Torus StaticMesh Assets", EUserInterfaceActionType::ToggleButton, FInputChord());
+	UI_COMMAND(BeginAddSphericalBoxPrimitiveTool, "SphrB", "Create new Sphere StaticMesh Assets with Box Parameterization", EUserInterfaceActionType::ToggleButton, FInputChord());
+	UI_COMMAND(BeginAddSpherePrimitiveTool, "SphrA", "Create new Sphere StaticMesh Assets with Lat/Long Parameterization", EUserInterfaceActionType::ToggleButton, FInputChord());
+
+	UI_COMMAND(BeginAddPatchTool, "Patch", "Add Patch", EUserInterfaceActionType::ToggleButton, FInputChord());
+	UI_COMMAND(BeginRevolveBoundaryTool, "CutRev", "Cut Mesh and Revolve Cut Section Curve into New Asset", EUserInterfaceActionType::ToggleButton, FInputChord());
+	UI_COMMAND(BeginDrawPolygonTool, "PolyExt", "Draw/Extrude Polygons to create new StaticMesh Assets", EUserInterfaceActionType::ToggleButton, FInputChord());
+	UI_COMMAND(BeginDrawPolyPathTool, "PathExt", "Draw/Extrude PolyPaths to create new StaticMesh Assets", EUserInterfaceActionType::ToggleButton, FInputChord());
+	UI_COMMAND(BeginDrawAndRevolveTool, "PathRev", "Draw and Revolve a PolyPath to create a new StaticMesh Assets", EUserInterfaceActionType::ToggleButton, FInputChord());
+	UI_COMMAND(BeginShapeSprayTool, "Spray", "Shape Spray", EUserInterfaceActionType::ToggleButton, FInputChord());
+
+	UI_COMMAND(BeginSculptMeshTool, "VSclpt", "Vertex Sculpting", EUserInterfaceActionType::ToggleButton, FInputChord());
+	UI_COMMAND(BeginPolyEditTool, "PolyEd", "Edit Mesh via PolyGroups", EUserInterfaceActionType::ToggleButton, FInputChord());
+	UI_COMMAND(BeginGroupEdgeInsertionTool, "EdgeIns", "Insert PolyGroup Edge", EUserInterfaceActionType::ToggleButton, FInputChord());
+	UI_COMMAND(BeginEdgeLoopInsertionTool, "LoopIns", "Insert PolyGroup Edge Loop", EUserInterfaceActionType::ToggleButton, FInputChord());
+	UI_COMMAND(BeginTriEditTool, "TriEd", "Edit Mesh via Triangles", EUserInterfaceActionType::ToggleButton, FInputChord());
+	UI_COMMAND(BeginPolyDeformTool, "PolyDef", "Deform Mesh via PolyGroups", EUserInterfaceActionType::ToggleButton, FInputChord());
+	UI_COMMAND(BeginSmoothMeshTool, "Smooth", "Smooth Mesh surface", EUserInterfaceActionType::ToggleButton, FInputChord());
+	UI_COMMAND(BeginOffsetMeshTool, "Offset", "Offset Mesh surface", EUserInterfaceActionType::ToggleButton, FInputChord());
+	UI_COMMAND(BeginDisplaceMeshTool, "Displce", "Displace Mesh surface with optional subdivision", EUserInterfaceActionType::ToggleButton, FInputChord());
+	UI_COMMAND(BeginMeshSpaceDeformerTool, "Warp", "Reshape Mesh using Space Deformers", EUserInterfaceActionType::ToggleButton, FInputChord());
+	UI_COMMAND(BeginTransformMeshesTool, "XForm", "Transform selected Meshes", EUserInterfaceActionType::ToggleButton, FInputChord());
+	UI_COMMAND(BeginEditPivotTool, "Pivot", "Edit Mesh Pivots", EUserInterfaceActionType::ToggleButton, FInputChord());
+	UI_COMMAND(BeginBakeTransformTool, "BakeRS", "Bake Scale/Rotation into Mesh Asset", EUserInterfaceActionType::ToggleButton, FInputChord());
+	UI_COMMAND(BeginCombineMeshesTool, "Append", "Combine Selection into new StaticMesh Asset", EUserInterfaceActionType::ToggleButton, FInputChord());
+	UI_COMMAND(BeginDuplicateMeshesTool, "Dupe", "Duplicate Selection into new StaticMesh Asset", EUserInterfaceActionType::ToggleButton, FInputChord());
+	UI_COMMAND(BeginAlignObjectsTool, "Align", "Align Objects", EUserInterfaceActionType::ToggleButton, FInputChord());
+
+	UI_COMMAND(BeginRemeshSculptMeshTool, "DSclpt", "Dynamic Mesh Sculpting", EUserInterfaceActionType::ToggleButton, FInputChord());
+	UI_COMMAND(BeginRemeshMeshTool, "Remesh", "Retriangulate Mesh", EUserInterfaceActionType::ToggleButton, FInputChord());
+	UI_COMMAND(BeginProjectToTargetTool, "Project", "Map/Remesh onto Target Mesh (second selection)", EUserInterfaceActionType::ToggleButton, FInputChord());
+	UI_COMMAND(BeginSimplifyMeshTool, "Simplfy", "Simplify Mesh", EUserInterfaceActionType::ToggleButton, FInputChord());
+	UI_COMMAND(BeginEditNormalsTool, "Nrmls", "Modify/Update Mesh Normals", EUserInterfaceActionType::ToggleButton, FInputChord());
+	UI_COMMAND(BeginEditTangentsTool, "Tngnts", "Update Mesh Tangents", EUserInterfaceActionType::ToggleButton, FInputChord());
+	UI_COMMAND(BeginRemoveOccludedTrianglesTool, "Jacket", "Remove Hidden Triangles from selected Meshes", EUserInterfaceActionType::ToggleButton, FInputChord());
+	UI_COMMAND(BeginUVProjectionTool, "Project", "Set UVs from Projection", EUserInterfaceActionType::ToggleButton, FInputChord());
+	UI_COMMAND(BeginUVLayoutTool, "Layout", "Re-Pack Existing UVs", EUserInterfaceActionType::ToggleButton, FInputChord());
+	UI_COMMAND(BeginVoxelMergeTool, "VoxMrg", "Merge Selected Meshes (Voxel-Based)", EUserInterfaceActionType::ToggleButton, FInputChord());
+	UI_COMMAND(BeginVoxelBooleanTool, "VoxBool", "Boolean Selected Meshes (Voxel-Based)", EUserInterfaceActionType::ToggleButton, FInputChord());
+	UI_COMMAND(BeginVoxelSolidifyTool, "VoxWrap", "Wrap Selected Meshes (Voxel-Based)", EUserInterfaceActionType::ToggleButton, FInputChord());
+	UI_COMMAND(BeginVoxelBlendTool, "VoxBlnd", "Blend Selected Meshes (Voxel-Based)", EUserInterfaceActionType::ToggleButton, FInputChord());
+	UI_COMMAND(BeginVoxelMorphologyTool, "VoxMrph", "Offset/Inset Selected Meshes (Voxel-Based)", EUserInterfaceActionType::ToggleButton, FInputChord());
+	UI_COMMAND(BeginSelfUnionTool, "Merge", "Self-Union Selected Meshes to resolve Self-Intersections", EUserInterfaceActionType::ToggleButton, FInputChord());
+	UI_COMMAND(BeginMeshBooleanTool, "Boolean", "Apply Boolean/CSG operation to selected Mesh pair", EUserInterfaceActionType::ToggleButton, FInputChord());
+	UI_COMMAND(BeginMeshTrimTool, "Trim", "Trim/Cut selected mesh with second mesh", EUserInterfaceActionType::ToggleButton, FInputChord());
+	UI_COMMAND(BeginBspConversionTool, "BSPConv", "Convert BSP to StaticMesh Asset", EUserInterfaceActionType::ToggleButton, FInputChord());
+	UI_COMMAND(BeginMeshToVolumeTool, "Msh2Vol", "Convert Mesh to Volume", EUserInterfaceActionType::ToggleButton, FInputChord());
+	UI_COMMAND(BeginVolumeToMeshTool, "Vol2Msh", "Convert Volume to new Mesh Asset", EUserInterfaceActionType::ToggleButton, FInputChord());
+	UI_COMMAND(BeginPlaneCutTool, "PlnCut", "Cut Selected Meshes with Plane", EUserInterfaceActionType::ToggleButton, FInputChord());
+	UI_COMMAND(BeginMirrorTool, "Mirror", "Mirror Selected Meshes", EUserInterfaceActionType::ToggleButton, FInputChord());
+	UI_COMMAND(BeginHoleFillTool, "HFill", "Fill Holes in Mesh", EUserInterfaceActionType::ToggleButton, FInputChord());
+	UI_COMMAND(BeginLatticeDeformerTool, "Grid", "Deform Mesh with 3D Lattice/Grid", EUserInterfaceActionType::ToggleButton, FInputChord());
+	UI_COMMAND(BeginPolygonCutTool, "PolyCut", "Cut Mesh with Extruded Polygon", EUserInterfaceActionType::ToggleButton, FInputChord());
+	UI_COMMAND(BeginGlobalUVGenerateTool, "AutoUV", "Auto-Unwrap and Pack UVs for Mesh", EUserInterfaceActionType::ToggleButton, FInputChord());
+	UI_COMMAND(BeginGroupUVGenerateTool, "Unwrap", "Recalculate UV unwrapping for Mesh regions", EUserInterfaceActionType::ToggleButton, FInputChord());
+	UI_COMMAND(BeginMeshSelectionTool, "TriSel", "Select and Edit Mesh Triangles", EUserInterfaceActionType::ToggleButton, FInputChord());
+
+	UI_COMMAND(BeginPhysicsInspectorTool, "PInspct", "Inspect Physics Geometry for selected Meshes", EUserInterfaceActionType::ToggleButton, FInputChord());
+	UI_COMMAND(BeginSetCollisionGeometryTool, "Msh2Coll", "Convert Selected Meshes to Simple Collision Geometry (for last Selected)", EUserInterfaceActionType::ToggleButton, FInputChord());
+	UI_COMMAND(BeginEditCollisionGeometryTool, "EditPhys", "Edit Simple Collision Geometry for selected Mesh", EUserInterfaceActionType::ToggleButton, FInputChord());
+	UI_COMMAND(BeginExtractCollisionGeometryTool, "Coll2Msh", "Convert Simple Collision Geometry to new Mesh Asset", EUserInterfaceActionType::ToggleButton, FInputChord());
+
+	UI_COMMAND(BeginMeshInspectorTool, "Inspct", "Inspect Mesh Attributes", EUserInterfaceActionType::ToggleButton, FInputChord());
+	UI_COMMAND(BeginWeldEdgesTool, "Weld", "Weld Overlapping Mesh Edges", EUserInterfaceActionType::ToggleButton, FInputChord());
+	UI_COMMAND(BeginPolyGroupsTool, "GenGrps", "Generate new PolyGroups", EUserInterfaceActionType::ToggleButton, FInputChord());
+	UI_COMMAND(BeginEditMeshMaterialsTool, "MatEd", "Assign Materials to Selected Triangles", EUserInterfaceActionType::ToggleButton, FInputChord());
+	UI_COMMAND(BeginTransformUVIslandsTool, "XForm", "Transform UV Islands in UV Space", EUserInterfaceActionType::ToggleButton, FInputChord());
+	UI_COMMAND(BeginMeshAttributePaintTool, "MapPnt", "Paint Attribute Maps", EUserInterfaceActionType::ToggleButton, FInputChord());
+	UI_COMMAND(BeginAttributeEditorTool, "AttrEd", "Edit/Configure Mesh Attributes", EUserInterfaceActionType::ToggleButton, FInputChord());
+	UI_COMMAND(BeginBakeMeshAttributeMapsTool, "BakeTx", "Bake Image Maps for Target Mesh (optionally from second Source Mesh)", EUserInterfaceActionType::ToggleButton, FInputChord());
+	UI_COMMAND(BeginUVSeamEditTool, "SeamEd", "Add UV Seams to Mesh", EUserInterfaceActionType::ToggleButton, FInputChord());
+
+	UI_COMMAND(BeginGroomToMeshTool, "Helmet", "Generate Helmet Mesh for Selected Groom", EUserInterfaceActionType::ToggleButton, FInputChord());
+	UI_COMMAND(BeginGenerateLODMeshesTool, "HlmLOD", "Generate LODS for Hair Helmet", EUserInterfaceActionType::ToggleButton, FInputChord());
+	UI_COMMAND(BeginGroomCardsToMeshTool, "CardsToMesh", "Hair Cards to Mesh Tool", EUserInterfaceActionType::ToggleButton, FInputChord());
+	UI_COMMAND(BeginGroomMeshToCardsTool, "MeshToCards", "Mesh to Hair Cards Tool", EUserInterfaceActionType::ToggleButton, FInputChord());
+	UI_COMMAND(BeginGroomCardsEditorTool, "CardsEd", "Edit Hair Cards", EUserInterfaceActionType::ToggleButton, FInputChord());
+	UI_COMMAND(BeginSubdividePolyTool, "SubD", "Subdivide Mesh Polygroups or Triangles", EUserInterfaceActionType::ToggleButton, FInputChord());
+
+	UI_COMMAND(BeginGenerateStaticMeshLODAssetTool, "AutoLOD", "Generate Static Mesh LOD Asset", EUserInterfaceActionType::ToggleButton, FInputChord());
+	UI_COMMAND(BeginMeshGroupPaintTool, "GrpPnt", "Paint New Mesh Polygroups", EUserInterfaceActionType::ToggleButton, FInputChord());
+
+
+	UI_COMMAND(BeginPolyModelTool_FaceSelect, "Faces", "PolyGroup Face Selection Tool", EUserInterfaceActionType::ToggleButton, FInputChord());
+	UI_COMMAND(BeginPolyModelTool_EdgeSelect, "Edges", "PolyGroup Edge Selection Tool", EUserInterfaceActionType::ToggleButton, FInputChord());
+	UI_COMMAND(BeginPolyModelTool_VertexSelect, "Verts", "PolyGroup Vertex Selection Tool", EUserInterfaceActionType::ToggleButton, FInputChord());
+	UI_COMMAND(BeginPolyModelTool_AllSelect, "Any", "PolyGroup Face/Edge/Vertex Selection Tool", EUserInterfaceActionType::ToggleButton, FInputChord());
+	UI_COMMAND(BeginPolyModelTool_LoopSelect, "Loops", "PolyGroup Loop Selection Tool", EUserInterfaceActionType::ToggleButton, FInputChord());
+	UI_COMMAND(BeginPolyModelTool_RingSelect, "Rings", "PolyGroup Ring Selection Tool", EUserInterfaceActionType::ToggleButton, FInputChord());
+
+	UI_COMMAND(BeginPolyModelTool_Extrude, "Extrude", "PolyGroup Extrude Tool", EUserInterfaceActionType::ToggleButton, FInputChord());
+	UI_COMMAND(BeginPolyModelTool_Offset, "Offset", "PolyGroup Offset Tool", EUserInterfaceActionType::ToggleButton, FInputChord());
+	UI_COMMAND(BeginPolyModelTool_Inset, "Inset", "PolyGroup Inset Tool", EUserInterfaceActionType::ToggleButton, FInputChord());
+	UI_COMMAND(BeginPolyModelTool_Outset, "Outset", "PolyGroup Outset Tool", EUserInterfaceActionType::ToggleButton, FInputChord());
+	UI_COMMAND(BeginPolyModelTool_CutFaces, "Cut", "PolyGroup Cut Faces Tool", EUserInterfaceActionType::ToggleButton, FInputChord());
 
 	UI_COMMAND(AcceptActiveTool, "Accept", "Accept the active tool", EUserInterfaceActionType::Button, FInputChord());
 	UI_COMMAND(CancelActiveTool, "Cancel", "Cancel the active tool", EUserInterfaceActionType::Button, FInputChord());
-	UI_COMMAND(CompleteActiveTool, "Complete", "Complete the active tool", EUserInterfaceActionType::Button, FInputChord());
+	UI_COMMAND(CompleteActiveTool, "Done", "Complete the active tool", EUserInterfaceActionType::Button, FInputChord());
 
 	// Note that passing a chord into one of these calls hooks the key press to the respective action. 
 	UI_COMMAND(AcceptOrCompleteActiveTool, "Accept or Complete", "Accept or complete the active tool", EUserInterfaceActionType::Button, FInputChord(EKeys::Enter));
diff --git a/Engine/Plugins/Experimental/ModelingToolsEditorMode/Source/ModelingToolsEditorMode/Public/ModelingToolsEditorModeToolkit.h b/Engine/Plugins/Experimental/ModelingToolsEditorMode/Source/ModelingToolsEditorMode/Public/ModelingToolsEditorModeToolkit.h
--- a/Engine/Plugins/Experimental/ModelingToolsEditorMode/Source/ModelingToolsEditorMode/Public/ModelingToolsEditorModeToolkit.h
+++ b/Engine/Plugins/Experimental/ModelingToolsEditorMode/Source/ModelingToolsEditorMode/Public/ModelingToolsEditorModeToolkit.h
@@ -78,6 +78,7 @@ private:
 	TSharedPtr<SButton> CancelButton;
 	TSharedPtr<SButton> CompletedButton;
 
+	TSharedPtr<SWidget> MakeAcceptCancelPanel();
 	TSharedPtr<SWidget> MakeAssetConfigPanel();
 
 	bool bShowRealtimeWarning = false;
