Modeling Tools Viewport Overlay

diff --git a/Engine/Plugins/Experimental/ModelingToolsEditorMode/Source/ModelingToolsEditorMode/Private/ModelingToolsEditorModeToolkit.cpp b/Engine/Plugins/Experimental/ModelingToolsEditorMode/Source/ModelingToolsEditorMode/Private/ModelingToolsEditorModeToolkit.cpp
--- a/Engine/Plugins/Experimental/ModelingToolsEditorMode/Source/ModelingToolsEditorMode/Private/ModelingToolsEditorModeToolkit.cpp
+++ b/Engine/Plugins/Experimental/ModelingToolsEditorMode/Source/ModelingToolsEditorMode/Private/ModelingToolsEditorModeToolkit.cpp
@@ -52,6 +52,8 @@ void FModelingToolsEditorModeToolkit::Init(const TSharedPtr<IToolkitHost>& InitT
 {
 	FModeToolkit::Init(InitToolkitHost, InOwningMode);
 
+	GetToolkitHost()->OnActiveViewportChanged().AddSP(this, &FModelingToolsEditorModeToolkit::OnActiveViewportChanged);
+
 	ModeWarningArea = SNew(STextBlock)
 		.AutoWrapText(true)
 		.Font(FCoreStyle::GetDefaultFontStyle("Bold", 9))
@@ -118,6 +120,59 @@ void FModelingToolsEditorModeToolkit::Init(const TSharedPtr<IToolkitHost>& InitT
 
 	GetScriptableEditorMode()->GetInteractiveToolsContext()->OnToolNotificationMessage.AddSP(this, &FModelingToolsEditorModeToolkit::PostNotification);
 	GetScriptableEditorMode()->GetInteractiveToolsContext()->OnToolWarningMessage.AddSP(this, &FModelingToolsEditorModeToolkit::PostWarning);
+
+	SAssignNew(ViewportOverlayWidget, SHorizontalBox)
+
+	+SHorizontalBox::Slot()
+	.HAlign(HAlign_Center)
+	.VAlign(VAlign_Bottom)
+	.Padding(FMargin(0.0f, 0.0f, 0.f, 15.f))
+	[
+		SNew(SBorder)
+		.BorderImage(FAppStyle::Get().GetBrush("EditorViewport.OverlayBrush"))
+		.Padding(8.f)
+		[
+			SNew(SHorizontalBox)
+
+			+SHorizontalBox::Slot()
+			.AutoWidth()
+			.VAlign(VAlign_Center)
+			.Padding(FMargin(0.f, 0.f, 8.f, 0.f))
+			[
+				SNew(STextBlock)
+				.Text(this, &FModelingToolsEditorModeToolkit::GetActiveToolDisplayName)
+			]
+
+			+SHorizontalBox::Slot()
+			.AutoWidth()
+			.Padding(FMargin(0.0, 0.f, 2.f, 0.f))
+			[
+				SNew(SButton)
+				.ButtonStyle(FAppStyle::Get(), "PrimaryButton")
+				.TextStyle( FAppStyle::Get(), "DialogButtonText" )
+				.Text(LOCTEXT("OverlayAccept", "Accept"))
+				.HAlign(HAlign_Center)
+				// .OnClicked(this, &FModelingToolsEditorModeToolkit::On... 
+				// .IsEnabled( this, &FModelingToolsEditorModeToolkit::On... 
+				// .Visibility( this, &FModelingToolsEditorModeToolkit::On... 
+			]
+
+			+SHorizontalBox::Slot()
+			.AutoWidth()
+			.Padding(FMargin(2.0, 0.f, 0.f, 0.f))
+			[
+				SNew(SButton)
+				.ButtonStyle(FAppStyle::Get(), "Button")
+				.TextStyle( FAppStyle::Get(), "DialogButtonText" )
+				.Text(LOCTEXT("OverlayCancel", "Cancel"))
+				.HAlign(HAlign_Center)
+				// .OnClicked(this, &FModelingToolsEditorModeToolkit::On... 
+				// .IsEnabled( this, &FModelingToolsEditorModeToolkit::On... 
+				// .Visibility( this, &FModelingToolsEditorModeToolkit::On... 
+			]
+		]	
+	];
+
 }
 
 
@@ -534,10 +589,17 @@ void FModelingToolsEditorModeToolkit::OnToolStarted(UInteractiveToolManager* Man
 
 	ModeHeaderArea->SetVisibility(EVisibility::Collapsed);
 	ActiveToolName = CurTool->GetToolInfo().ToolDisplayName;
+
+	GetToolkitHost()->AddViewportOverlayWidget(ViewportOverlayWidget.ToSharedRef());
 }
 
 void FModelingToolsEditorModeToolkit::OnToolEnded(UInteractiveToolManager* Manager, UInteractiveTool* Tool)
 {
+	if (IsHosted())
+	{
+		GetToolkitHost()->RemoveViewportOverlayWidget(ViewportOverlayWidget.ToSharedRef());
+	}
+
 	ModeDetailsView->SetObject(nullptr);
 	ActiveToolName = FText::GetEmpty();
 	ModeHeaderArea->SetVisibility(EVisibility::Visible);
@@ -551,6 +613,22 @@ void FModelingToolsEditorModeToolkit::OnToolEnded(UInteractiveToolManager* Manag
 	}
 }
 
+void FModelingToolsEditorModeToolkit::OnActiveViewportChanged(TSharedPtr<IAssetViewport> OldViewport, TSharedPtr<IAssetViewport> NewViewport)
+{
+	// Only worry about handling this notification if Modeling has an active tool
+	if (!ActiveToolName.IsEmpty())
+	{
+		// Check first to see if this changed because the old viewport was deleted and if not, remove our hud
+		if (OldViewport)	
+		{
+			GetToolkitHost()->RemoveViewportOverlayWidget(ViewportOverlayWidget.ToSharedRef(), OldViewport);
+		}
+
+		// Add the hud to the new viewport
+		GetToolkitHost()->AddViewportOverlayWidget(ViewportOverlayWidget.ToSharedRef(), NewViewport);
+	}
+}
+
 void FModelingToolsEditorModeToolkit::UpdateShowWarnings()
 {
 	if (bShowRealtimeWarning )
diff --git a/Engine/Plugins/Experimental/ModelingToolsEditorMode/Source/ModelingToolsEditorMode/Public/ModelingToolsEditorModeToolkit.h b/Engine/Plugins/Experimental/ModelingToolsEditorMode/Source/ModelingToolsEditorMode/Public/ModelingToolsEditorModeToolkit.h
--- a/Engine/Plugins/Experimental/ModelingToolsEditorMode/Source/ModelingToolsEditorMode/Public/ModelingToolsEditorModeToolkit.h
+++ b/Engine/Plugins/Experimental/ModelingToolsEditorMode/Source/ModelingToolsEditorMode/Public/ModelingToolsEditorModeToolkit.h
@@ -58,6 +58,7 @@ public:
 	virtual void OnToolStarted(UInteractiveToolManager* Manager, UInteractiveTool* Tool) override;
 	virtual void OnToolEnded(UInteractiveToolManager* Manager, UInteractiveTool* Tool) override;
 
+	void OnActiveViewportChanged(TSharedPtr<IAssetViewport>, TSharedPtr<IAssetViewport> );
 
 
 private:
@@ -70,6 +71,7 @@ private:
 	TSharedPtr<SWidget> ToolkitWidget;
 	void UpdateActiveToolProperties();
 
+	TSharedPtr<SWidget> ViewportOverlayWidget;
 
 	TSharedPtr<STextBlock> ModeWarningArea;
 	TSharedPtr<STextBlock> ModeHeaderArea;
