ModelingMode UX: hook up accept/cancel/completed buttons in new viewport overlay
#rnx

diff --git a/Engine/Plugins/Experimental/ModelingToolsEditorMode/Source/ModelingToolsEditorMode/Private/ModelingToolsEditorModeToolkit.cpp b/Engine/Plugins/Experimental/ModelingToolsEditorMode/Source/ModelingToolsEditorMode/Private/ModelingToolsEditorModeToolkit.cpp
--- a/Engine/Plugins/Experimental/ModelingToolsEditorMode/Source/ModelingToolsEditorMode/Private/ModelingToolsEditorModeToolkit.cpp
+++ b/Engine/Plugins/Experimental/ModelingToolsEditorMode/Source/ModelingToolsEditorMode/Private/ModelingToolsEditorModeToolkit.cpp
@@ -101,11 +101,6 @@ void FModelingToolsEditorModeToolkit::Init(const TSharedPtr<IToolkitHost>& InitT
 					ModeDetailsView->AsShared()
 				]
 
-			+ SVerticalBox::Slot().AutoHeight().HAlign(HAlign_Fill).Padding(5)
-				[
-					MakeAcceptCancelPanel()->AsShared()
-				]
-
 			+ SVerticalBox::Slot().AutoHeight().HAlign(HAlign_Fill).Padding(5)
 				[
 					MakeAssetConfigPanel()->AsShared()
@@ -152,9 +147,9 @@ void FModelingToolsEditorModeToolkit::Init(const TSharedPtr<IToolkitHost>& InitT
 				.TextStyle( FAppStyle::Get(), "DialogButtonText" )
 				.Text(LOCTEXT("OverlayAccept", "Accept"))
 				.HAlign(HAlign_Center)
-				// .OnClicked(this, &FModelingToolsEditorModeToolkit::On... 
-				// .IsEnabled( this, &FModelingToolsEditorModeToolkit::On... 
-				// .Visibility( this, &FModelingToolsEditorModeToolkit::On... 
+				.OnClicked_Lambda([this]() { GetScriptableEditorMode()->GetInteractiveToolsContext()->EndTool(EToolShutdownType::Accept); return FReply::Handled(); })
+				.IsEnabled_Lambda([this]() { return GetScriptableEditorMode()->GetInteractiveToolsContext()->CanAcceptActiveTool(); })
+				.Visibility_Lambda([this]() { return GetScriptableEditorMode()->GetInteractiveToolsContext()->ActiveToolHasAccept() ? EVisibility::Visible : EVisibility::Collapsed; })
 			]
 
 			+SHorizontalBox::Slot()
@@ -166,9 +161,23 @@ void FModelingToolsEditorModeToolkit::Init(const TSharedPtr<IToolkitHost>& InitT
 				.TextStyle( FAppStyle::Get(), "DialogButtonText" )
 				.Text(LOCTEXT("OverlayCancel", "Cancel"))
 				.HAlign(HAlign_Center)
-				// .OnClicked(this, &FModelingToolsEditorModeToolkit::On... 
-				// .IsEnabled( this, &FModelingToolsEditorModeToolkit::On... 
-				// .Visibility( this, &FModelingToolsEditorModeToolkit::On... 
+				.OnClicked_Lambda([this]() { GetScriptableEditorMode()->GetInteractiveToolsContext()->EndTool(EToolShutdownType::Cancel); return FReply::Handled(); })
+				.IsEnabled_Lambda([this]() { return GetScriptableEditorMode()->GetInteractiveToolsContext()->CanCancelActiveTool(); })
+				.Visibility_Lambda([this]() { return GetScriptableEditorMode()->GetInteractiveToolsContext()->ActiveToolHasAccept() ? EVisibility::Visible : EVisibility::Collapsed; })
+			]
+
+			+SHorizontalBox::Slot()
+			.AutoWidth()
+			.Padding(FMargin(2.0, 0.f, 0.f, 0.f))
+			[
+				SNew(SButton)
+				.ButtonStyle(FAppStyle::Get(), "Button")
+				.TextStyle( FAppStyle::Get(), "DialogButtonText" )
+				.Text(LOCTEXT("OverlayComplete", "Complete"))
+				.HAlign(HAlign_Center)
+				.OnClicked_Lambda([this]() { GetScriptableEditorMode()->GetInteractiveToolsContext()->EndTool(EToolShutdownType::Completed); return FReply::Handled(); })
+				.IsEnabled_Lambda([this]() { return GetScriptableEditorMode()->GetInteractiveToolsContext()->CanCompleteActiveTool(); })
+				.Visibility_Lambda([this]() { return GetScriptableEditorMode()->GetInteractiveToolsContext()->CanCompleteActiveTool() ? EVisibility::Visible : EVisibility::Collapsed; })
 			]
 		]	
 	];
@@ -177,42 +186,6 @@ void FModelingToolsEditorModeToolkit::Init(const TSharedPtr<IToolkitHost>& InitT
 
 
 
-
-TSharedPtr<SWidget> FModelingToolsEditorModeToolkit::MakeAcceptCancelPanel()
-{
-
-	TSharedPtr<FUICommandList> CommandList = GetToolkitCommands();
-	FName ToolbarCustomizationName = GetEditorMode() ? GetEditorMode()->GetModeInfo().ToolbarCustomizationName : GetScriptableEditorMode()->GetModeInfo().ToolbarCustomizationName;
-	FUniformToolBarBuilder RowBuilder(CommandList, FMultiBoxCustomization(ToolbarCustomizationName));
-	RowBuilder.SetStyle(&FEditorStyle::Get(), "PaletteToolBar");
-
-	const FModelingToolsManagerCommands& Commands = FModelingToolsManagerCommands::Get();
-	RowBuilder.AddToolBarButton(Commands.AcceptActiveTool);
-	RowBuilder.AddToolBarButton(Commands.CancelActiveTool);
-	RowBuilder.AddToolBarButton(Commands.CompleteActiveTool);
-
-	TSharedPtr<SVerticalBox> Content = SNew(SVerticalBox)
-	+ SVerticalBox::Slot().AutoHeight().HAlign(HAlign_Right).Padding(0, 4, 0, 10)
-		[
-			RowBuilder.MakeWidget()
-		];
-
-	TSharedPtr<SExpandableArea> AcceptCancelPanel = SNew(SExpandableArea)
-		.HeaderPadding(FMargin(2.0f))
-		.Padding(FMargin(2.f))
-		.BorderImage(FEditorStyle::Get().GetBrush("DetailsView.CategoryTop"))
-		.BorderBackgroundColor(FLinearColor(.6, .6, .6, 1.0f))
-		.BodyBorderBackgroundColor(FLinearColor::Transparent)
-		.AreaTitleFont(FEditorStyle::Get().GetFontStyle("EditorModesPanel.CategoryFontStyle"))
-		.BodyContent()
-		[
-			Content->AsShared()
-		];
-
-	return AcceptCancelPanel;
-}
-
-
 TSharedPtr<SWidget> FModelingToolsEditorModeToolkit::MakeAssetConfigPanel()
 {
 	AssetLocationModes.Add(MakeShared<FString>(TEXT("AutoGen Folder (World-Relative)")));
diff --git a/Engine/Plugins/Experimental/ModelingToolsEditorMode/Source/ModelingToolsEditorMode/Public/ModelingToolsEditorModeToolkit.h b/Engine/Plugins/Experimental/ModelingToolsEditorMode/Source/ModelingToolsEditorMode/Public/ModelingToolsEditorModeToolkit.h
--- a/Engine/Plugins/Experimental/ModelingToolsEditorMode/Source/ModelingToolsEditorMode/Public/ModelingToolsEditorModeToolkit.h
+++ b/Engine/Plugins/Experimental/ModelingToolsEditorMode/Source/ModelingToolsEditorMode/Public/ModelingToolsEditorModeToolkit.h
@@ -80,7 +80,6 @@ private:
 	TSharedPtr<SButton> CancelButton;
 	TSharedPtr<SButton> CompletedButton;
 
-	TSharedPtr<SWidget> MakeAcceptCancelPanel();
 	TSharedPtr<SWidget> MakeAssetConfigPanel();
 
 	bool bShowRealtimeWarning = false;
@@ -96,4 +95,7 @@ private:
 	void OnAssetSettingsModified();
 	FDelegateHandle AssetSettingsModifiedHandle;
 	void OnShowAssetSettings();
+
+
+	// custom accept/cancel/complete handlers
 };
