ModelingMode UX: display tool help text in status bar when tool is active
#rnx

diff --git a/Engine/Plugins/Experimental/ModelingToolsEditorMode/Source/ModelingToolsEditorMode/ModelingToolsEditorMode.Build.cs b/Engine/Plugins/Experimental/ModelingToolsEditorMode/Source/ModelingToolsEditorMode/ModelingToolsEditorMode.Build.cs
--- a/Engine/Plugins/Experimental/ModelingToolsEditorMode/Source/ModelingToolsEditorMode/ModelingToolsEditorMode.Build.cs
+++ b/Engine/Plugins/Experimental/ModelingToolsEditorMode/Source/ModelingToolsEditorMode/ModelingToolsEditorMode.Build.cs
@@ -46,6 +46,7 @@ public class ModelingToolsEditorMode : ModuleRules
 				"UnrealEd",
                 "ContentBrowser",
                 "LevelEditor",
+				"StatusBar",
 				"EditorStyle",
                 "Projects",
 				"MeshDescription",
diff --git a/Engine/Plugins/Experimental/ModelingToolsEditorMode/Source/ModelingToolsEditorMode/Private/ModelingToolsEditorModeToolkit.cpp b/Engine/Plugins/Experimental/ModelingToolsEditorMode/Source/ModelingToolsEditorMode/Private/ModelingToolsEditorModeToolkit.cpp
--- a/Engine/Plugins/Experimental/ModelingToolsEditorMode/Source/ModelingToolsEditorMode/Private/ModelingToolsEditorModeToolkit.cpp
+++ b/Engine/Plugins/Experimental/ModelingToolsEditorMode/Source/ModelingToolsEditorMode/Private/ModelingToolsEditorModeToolkit.cpp
@@ -304,12 +304,21 @@ void FModelingToolsEditorModeToolkit::UpdateActiveToolProperties()
 
 void FModelingToolsEditorModeToolkit::PostNotification(const FText& Message)
 {
+	ClearNotification();
+
 	ActiveToolMessage = Message;
+
+	const FName LevelEditorStatusBarName = "LevelEditor.StatusBar";
+	ActiveToolMessageHandle = GEditor->GetEditorSubsystem<UStatusBarSubsystem>()->PushStatusBarMessage(LevelEditorStatusBarName, ActiveToolMessage);
 }
 
 void FModelingToolsEditorModeToolkit::ClearNotification()
 {
 	ActiveToolMessage = FText::GetEmpty();
+
+	const FName LevelEditorStatusBarName = "LevelEditor.StatusBar";
+	GEditor->GetEditorSubsystem<UStatusBarSubsystem>()->PopStatusBarMessage(LevelEditorStatusBarName, ActiveToolMessageHandle);
+	ActiveToolMessageHandle.Reset();
 }
 
 
diff --git a/Engine/Plugins/Experimental/ModelingToolsEditorMode/Source/ModelingToolsEditorMode/Public/ModelingToolsEditorModeToolkit.h b/Engine/Plugins/Experimental/ModelingToolsEditorMode/Source/ModelingToolsEditorMode/Public/ModelingToolsEditorModeToolkit.h
--- a/Engine/Plugins/Experimental/ModelingToolsEditorMode/Source/ModelingToolsEditorMode/Public/ModelingToolsEditorModeToolkit.h
+++ b/Engine/Plugins/Experimental/ModelingToolsEditorMode/Source/ModelingToolsEditorMode/Public/ModelingToolsEditorModeToolkit.h
@@ -8,9 +8,9 @@
 #include "ModelingToolsEditorMode.h"
 #include "Widgets/Input/SComboBox.h"
 #include "Widgets/Input/STextComboBox.h"
-
-
 #include "Widgets/SBoxPanel.h"
+#include "StatusBarSubsystem.h"
+
 
 class IDetailsView;
 class SButton;
@@ -67,6 +67,7 @@ private:
 
 	FText ActiveToolName;
 	FText ActiveToolMessage;
+	FStatusBarMessageHandle ActiveToolMessageHandle;
 
 	TSharedPtr<SWidget> ToolkitWidget;
 	void UpdateActiveToolProperties();
