ModelingMode: Support for HiRes SourceModel in AutoLOD Process/Tool. Copies input LOD0 to HiRes on Accept, and uses input HiRes if available, so AutoLOD can be repeated cleanly. Extend ComponentTargetFactory registration API to allow removal and querying for registered Factory. Add active Editing LOD to FStaticMeshComponentTargetFactory. Expose drop-down in ModelingMode UI to configure the active Editing LOD level. Remove Save settings from Modeling Mode Quick-Access UI.
#rnx

diff --git a/Engine/Plugins/Experimental/MeshLODToolset/Source/MeshLODToolset/Private/Graphs/GenerateStaticMeshLODProcess.cpp b/Engine/Plugins/Experimental/MeshLODToolset/Source/MeshLODToolset/Private/Graphs/GenerateStaticMeshLODProcess.cpp
--- a/Engine/Plugins/Experimental/MeshLODToolset/Source/MeshLODToolset/Private/Graphs/GenerateStaticMeshLODProcess.cpp
+++ b/Engine/Plugins/Experimental/MeshLODToolset/Source/MeshLODToolset/Private/Graphs/GenerateStaticMeshLODProcess.cpp
@@ -63,7 +63,10 @@ bool UGenerateStaticMeshLODProcess::Initialize(UStaticMesh* StaticMeshIn)
 	FlushRenderingCommands();
 
 	SourceStaticMesh = StaticMeshIn;
-	const FMeshDescription* SourceMeshDescription = StaticMeshIn->GetMeshDescription(0);
+
+	bool bHasHiResSource = SourceStaticMesh->IsHiResMeshDescriptionValid();
+	const FMeshDescription* SourceMeshDescription =
+		(bHasHiResSource) ? SourceStaticMesh->GetHiResMeshDescription() : SourceStaticMesh->GetMeshDescription(0);
 
 	// start async mesh-conversion
 	SourceMesh.Clear();
@@ -495,7 +498,7 @@ bool UGenerateStaticMeshLODProcess::WriteDerivedAssetData()
 }
 
 
-void UGenerateStaticMeshLODProcess::UpdateSourceAsset()
+void UGenerateStaticMeshLODProcess::UpdateSourceAsset(bool bSetNewHDSourceAsset)
 {
 	AllDerivedTextures.Reset();
 
@@ -503,7 +506,7 @@ void UGenerateStaticMeshLODProcess::UpdateSourceAsset()
 
 	WriteDerivedMaterials();
 
-	UpdateSourceStaticMeshAsset();
+	UpdateSourceStaticMeshAsset(bSetNewHDSourceAsset);
 
 	// clear list of derived textures we were holding onto to prevent GC
 	AllDerivedTextures.Reset();
@@ -652,6 +655,11 @@ void UGenerateStaticMeshLODProcess::WriteDerivedMaterials()
 		FMaterialInfo& DerivedMaterialInfo = DerivedMaterials[mi];
 
 		UMaterialInterface* MaterialInterface = SourceMaterialInfo.SourceMaterial.MaterialInterface;
+		if (MaterialInterface == nullptr)
+		{
+			DerivedMaterialInfo.SourceMaterial.MaterialInterface = nullptr;
+			continue;
+		}
 		bool bSourceIsMIC = (Cast<UMaterialInstanceConstant>(MaterialInterface) != nullptr);
 
 		FString SourceMaterialPath = UEditorAssetLibrary::GetPathNameForLoadedAsset(MaterialInterface);
@@ -805,12 +813,20 @@ void UGenerateStaticMeshLODProcess::WriteDerivedStaticMeshAsset()
 
 
 
-void UGenerateStaticMeshLODProcess::UpdateSourceStaticMeshAsset()
+void UGenerateStaticMeshLODProcess::UpdateSourceStaticMeshAsset(bool bSetNewHDSourceAsset)
 {
 	GEditor->BeginTransaction(LOCTEXT("UpdateExistingAssetMessage", "Added Generated LOD"));
 
 	SourceStaticMesh->Modify();
 
+	if (bSetNewHDSourceAsset)
+	{
+		FMeshDescription* NewSourceMD = SourceStaticMesh->CreateHiResMeshDescription();
+		const FMeshDescription* ExistingSourceMD = SourceStaticMesh->GetMeshDescription(0);
+		*NewSourceMD = *ExistingSourceMD;
+		SourceStaticMesh->CommitHiResMeshDescription();
+	}
+
 	TArray<FStaticMaterial> ExistingMaterials = SourceStaticMesh->GetStaticMaterials();
 
 	// append new materials to material set
@@ -835,20 +851,20 @@ void UGenerateStaticMeshLODProcess::UpdateSourceStaticMeshAsset()
 	// update materials on generated mesh
 	SourceStaticMesh->SetStaticMaterials(ExistingMaterials);
 
-	// store new derived LOD as LOD 1
-	SourceStaticMesh->SetNumSourceModels(2);
-	FMeshDescription* MeshDescription = SourceStaticMesh->GetMeshDescription(1);
+	// store new derived LOD as LOD 0
+	SourceStaticMesh->SetNumSourceModels(1);
+	FMeshDescription* MeshDescription = SourceStaticMesh->GetMeshDescription(0);
 	if (MeshDescription == nullptr)
 	{
-		MeshDescription = SourceStaticMesh->CreateMeshDescription(1);
+		MeshDescription = SourceStaticMesh->CreateMeshDescription(0);
 	}
 	FConversionToMeshDescriptionOptions ConversionOptions;
 	FDynamicMeshToMeshDescription Converter(ConversionOptions);
 	Converter.Convert(&DerivedLODMesh, *MeshDescription);
-	SourceStaticMesh->CommitMeshDescription(1);
+	SourceStaticMesh->CommitMeshDescription(0);
 
 	// this will prevent simplification?
-	FStaticMeshSourceModel& SrcModel = SourceStaticMesh->GetSourceModel(1);
+	FStaticMeshSourceModel& SrcModel = SourceStaticMesh->GetSourceModel(0);
 	SrcModel.ReductionSettings.MaxDeviation = 0.0f;
 	SrcModel.ReductionSettings.PercentTriangles = 1.0f;
 	SrcModel.ReductionSettings.PercentVertices = 1.0f;
diff --git a/Engine/Plugins/Experimental/MeshLODToolset/Source/MeshLODToolset/Private/Tools/GenerateStaticMeshLODAssetTool.cpp b/Engine/Plugins/Experimental/MeshLODToolset/Source/MeshLODToolset/Private/Tools/GenerateStaticMeshLODAssetTool.cpp
--- a/Engine/Plugins/Experimental/MeshLODToolset/Source/MeshLODToolset/Private/Tools/GenerateStaticMeshLODAssetTool.cpp
+++ b/Engine/Plugins/Experimental/MeshLODToolset/Source/MeshLODToolset/Private/Tools/GenerateStaticMeshLODAssetTool.cpp
@@ -409,7 +409,13 @@ void UGenerateStaticMeshLODAssetTool::UpdateExistingAsset()
 	GenerateProcess->CalculateDerivedPathName(BasicProperties->GeneratedSuffix);
 
 	check(GenerateProcess->GraphEvalCriticalSection.TryLock());		// No ops should be running
-	GenerateProcess->UpdateSourceAsset();
+
+	// only updated HD source if we have no HD source asset. Otherwise we are overwriting with existing lowpoly LOD0.
+	bool bUpdateHDSource =
+		BasicProperties->bSaveAsHDSource &&
+		(GenerateProcess->GetSourceStaticMesh()->IsHiResMeshDescriptionValid() == false);
+
+	GenerateProcess->UpdateSourceAsset(bUpdateHDSource);
 	GenerateProcess->GraphEvalCriticalSection.Unlock();
 }
 
diff --git a/Engine/Plugins/Experimental/MeshLODToolset/Source/MeshLODToolset/Public/Graphs/GenerateStaticMeshLODProcess.h b/Engine/Plugins/Experimental/MeshLODToolset/Source/MeshLODToolset/Public/Graphs/GenerateStaticMeshLODProcess.h
--- a/Engine/Plugins/Experimental/MeshLODToolset/Source/MeshLODToolset/Public/Graphs/GenerateStaticMeshLODProcess.h
+++ b/Engine/Plugins/Experimental/MeshLODToolset/Source/MeshLODToolset/Public/Graphs/GenerateStaticMeshLODProcess.h
@@ -178,13 +178,13 @@ public:
 	/**
 	 * Creates new Asset
 	 */
-	bool WriteDerivedAssetData();
+	virtual bool WriteDerivedAssetData();
 
 
 	/**
 	 * Updates existing SM Asset
 	 */
-	void UpdateSourceAsset();
+	virtual void UpdateSourceAsset(bool bSetNewHDSourceAsset = false);
 
 
 
@@ -264,7 +264,7 @@ protected:
 	void UpdateMaterialTextureParameters(UMaterialInstanceConstant* Material, FMaterialInfo& DerivedMaterialInfo);
 	void WriteDerivedStaticMeshAsset();
 
-	void UpdateSourceStaticMeshAsset();
+	void UpdateSourceStaticMeshAsset(bool bSetNewHDSourceAsset);
 
 	void UpdateMaterialTextureParameters(UMaterialInstanceDynamic* Material, const FMaterialInfo& SourceMaterialInfo, 
 		const TMap<UTexture2D*,UTexture2D*>& PreviewTextures, UTexture2D* PreviewNormalMap);
diff --git a/Engine/Plugins/Experimental/MeshLODToolset/Source/MeshLODToolset/Public/Tools/GenerateStaticMeshLODAssetTool.h b/Engine/Plugins/Experimental/MeshLODToolset/Source/MeshLODToolset/Public/Tools/GenerateStaticMeshLODAssetTool.h
--- a/Engine/Plugins/Experimental/MeshLODToolset/Source/MeshLODToolset/Public/Tools/GenerateStaticMeshLODAssetTool.h
+++ b/Engine/Plugins/Experimental/MeshLODToolset/Source/MeshLODToolset/Public/Tools/GenerateStaticMeshLODAssetTool.h
@@ -69,6 +69,11 @@ public:
 	UPROPERTY(EditAnywhere, Category = AssetOptions, meta = (TransientToolProperty))
 	FString GeneratedSuffix;
 
+	/** If true, the high-resolution input mesh is stored as HD source mesh in the Asset */
+	UPROPERTY(EditAnywhere, Category = AssetOptions, meta = (EditCondition = "OutputMode == EGenerateLODAssetOutputMode::UpdateExistingAsset"))
+	bool bSaveAsHDSource = true;
+
+
 	UPROPERTY(EditAnywhere, Category = Settings)
 	bool bParallelExecution = false;
 
diff --git a/Engine/Plugins/Experimental/ModelingToolsEditorMode/Source/ModelingToolsEditorMode/Private/ModelingToolsEditorModeToolkit.cpp b/Engine/Plugins/Experimental/ModelingToolsEditorMode/Source/ModelingToolsEditorMode/Private/ModelingToolsEditorModeToolkit.cpp
--- a/Engine/Plugins/Experimental/ModelingToolsEditorMode/Source/ModelingToolsEditorMode/Private/ModelingToolsEditorModeToolkit.cpp
+++ b/Engine/Plugins/Experimental/ModelingToolsEditorMode/Source/ModelingToolsEditorMode/Private/ModelingToolsEditorModeToolkit.cpp
@@ -20,6 +20,10 @@
 #include "Widgets/Images/SImage.h"
 #include "Widgets/Layout/SExpandableArea.h"
 
+// for LOD setting
+#include "EditorInteractiveToolsFrameworkModule.h"
+#include "Tools/EditorComponentSourceFactory.h"
+
 #define LOCTEXT_NAMESPACE "FModelingToolsEditorModeToolkit"
 
 
@@ -191,6 +195,10 @@ void FModelingToolsEditorModeToolkit::Init(const TSharedPtr<IToolkitHost>& InitT
 
 TSharedPtr<SWidget> FModelingToolsEditorModeToolkit::MakeAssetConfigPanel()
 {
+	//
+	// New Asset Location drop-down
+	//
+
 	AssetLocationModes.Add(MakeShared<FString>(TEXT("AutoGen Folder (World-Relative)")));
 	AssetLocationModes.Add(MakeShared<FString>(TEXT("AutoGen Folder (Global)")));
 	AssetLocationModes.Add(MakeShared<FString>(TEXT("Current Folder")));
@@ -212,21 +220,85 @@ TSharedPtr<SWidget> FModelingToolsEditorModeToolkit::MakeAssetConfigPanel()
 	AssetSettingsModifiedHandle = Settings->OnModified.AddLambda([this](UObject*, FProperty*) { OnAssetSettingsModified(); });
 
 
+	//
+	// LOD selection dropdown
+	//
+
+	AssetLODModes.Add(MakeShared<FString>(TEXT("Max Available")));
+	AssetLODModes.Add(MakeShared<FString>(TEXT("HiRes")));
+	AssetLODModes.Add(MakeShared<FString>(TEXT("LOD0")));
+	AssetLODModes.Add(MakeShared<FString>(TEXT("LOD1")));
+	AssetLODModes.Add(MakeShared<FString>(TEXT("LOD2")));
+	AssetLODModes.Add(MakeShared<FString>(TEXT("LOD3")));
+	AssetLODModes.Add(MakeShared<FString>(TEXT("LOD4")));
+	AssetLODModes.Add(MakeShared<FString>(TEXT("LOD5")));
+	AssetLODModes.Add(MakeShared<FString>(TEXT("LOD6")));
+	AssetLODModes.Add(MakeShared<FString>(TEXT("LOD7")));
+	AssetLODMode = SNew(STextComboBox)
+		.OptionsSource(&AssetLODModes)
+		.OnSelectionChanged_Lambda([&](TSharedPtr<FString> String, ESelectInfo::Type)
+	{
+		if (FEditorInteractiveToolsFrameworkGlobals::RegisteredStaticMeshTargetFactoryKey >= 0)
+		{
+			FComponentTargetFactory* Factory = FindComponentTargetFactoryByKey(FEditorInteractiveToolsFrameworkGlobals::RegisteredStaticMeshTargetFactoryKey);
+			if (Factory != nullptr)
+			{
+				FStaticMeshComponentTargetFactory* StaticMeshFactory = static_cast<FStaticMeshComponentTargetFactory*>(Factory);
+				if (*String == *AssetLODModes[0])
+				{
+					StaticMeshFactory->CurrentEditingLOD = EStaticMeshEditingLOD::MaxQuality;
+				}
+				else if (*String == *AssetLODModes[1])
+				{
+					StaticMeshFactory->CurrentEditingLOD = EStaticMeshEditingLOD::HiResSource;
+				}
+				else
+				{
+					for (int32 k = 2; k < AssetLODModes.Num(); ++k)
+					{
+						if (*String == *AssetLODModes[k])
+						{
+							StaticMeshFactory->CurrentEditingLOD = (EStaticMeshEditingLOD)(k - 2);
+							break;
+						}
+					}
+				}
+			}
+		}
+	});
+
+
+	// initialize combos
+	//UpdateAssetPanelFromSettings();
+	AssetLODMode->SetSelectedItem(AssetLODModes[0]);		// TODO read from settings?
+
+
+	AssetLODModeLabel = SNew(STextBlock).Text(LOCTEXT("ActiveLODLabel", "Editing LOD"));
+
 	TSharedPtr<SVerticalBox> Content = SNew(SVerticalBox)
-	//+ SVerticalBox::Slot().AutoHeight().HAlign(HAlign_Fill).Padding(0)
-	//	[
-	//		SNew(SSeparator)
-	//	]
-	+ SVerticalBox::Slot().AutoHeight().HAlign(HAlign_Fill).Padding(0, 4, 0, 5)
+	+ SVerticalBox::Slot().MaxHeight(20).HAlign(HAlign_Fill)
 		[
 			SNew(SHorizontalBox)
-			+ SHorizontalBox::Slot().Padding(0).HAlign(HAlign_Left).VAlign(VAlign_Center).AutoWidth()
+			+ SHorizontalBox::Slot().Padding(0, 2, 2, 2).HAlign(HAlign_Right).VAlign(VAlign_Center).AutoWidth()
+				[ AssetLODModeLabel->AsShared() ]
+			+ SHorizontalBox::Slot().Padding(0).FillWidth(4.0f)
+				[ AssetLODMode->AsShared() ]
+		]
+
+	+ SVerticalBox::Slot().MaxHeight(20).HAlign(HAlign_Fill).Padding(0, 3, 0, 0)
+		[
+			SNew(SHorizontalBox)
+			+ SHorizontalBox::Slot().Padding(0, 2, 2, 2).HAlign(HAlign_Right).VAlign(VAlign_Center).AutoWidth()
 				[
 					SNew(STextBlock)
-					.Text(LOCTEXT("AssetPanelHeaderLabel", "Generated Asset Settings"))
-					.Font(FCoreStyle::GetDefaultFontStyle("Bold", 9))
+					.Text(LOCTEXT("AssetLocationLabel", "New Asset Location"))
+				]
+			+ SHorizontalBox::Slot().Padding(0).FillWidth(4.0f)
+				[
+					AssetLocationMode->AsShared()
 				]
-			+ SHorizontalBox::Slot().HAlign(HAlign_Right).Padding(0).FillWidth(1.0f)
+
+			+ SHorizontalBox::Slot().HAlign(HAlign_Right).Padding(5, 0, 0, 0).FillWidth(1.0f)
 				[
 					SNew(SBox).MaxDesiredHeight(16).MaxDesiredWidth(17)
 					[
@@ -235,37 +307,12 @@ TSharedPtr<SWidget> FModelingToolsEditorModeToolkit::MakeAssetConfigPanel()
 						.ButtonStyle(FEditorStyle::Get(), "FlatButton.Default")
 						.OnClicked_Lambda( [this]() {OnShowAssetSettings(); return FReply::Handled(); } )
 						[
-							SNew(SImage)
-							.Image(FSlateIcon(FEditorStyle::GetStyleSetName(), "FoliageEditMode.Settings").GetIcon())
+							SNew(SImage).Image(FSlateIcon(FEditorStyle::GetStyleSetName(), "FoliageEditMode.Settings").GetIcon())
 						]
 					]
 				]
-		]
-	+ SVerticalBox::Slot().MaxHeight(20).HAlign(HAlign_Fill).Padding(0)
-		[
-			SNew(SHorizontalBox)
-			+ SHorizontalBox::Slot().Padding(0, 2, 2, 2).HAlign(HAlign_Right).VAlign(VAlign_Center).AutoWidth()
-				[
-					SNew(STextBlock)
-					.Text(LOCTEXT("AssetLocationLabel", "Location"))
-				]
-			+ SHorizontalBox::Slot().Padding(0).FillWidth(4.0f)
-				[
-					AssetLocationMode->AsShared()
-				]
-			+ SHorizontalBox::Slot().Padding(10, 2, 2, 2).HAlign(HAlign_Right).VAlign(VAlign_Center).AutoWidth()
-				[
-					SNew(STextBlock)
-					.Text(LOCTEXT("AssetSaveModeLabel", "Save Mode"))
-				]
-			+ SHorizontalBox::Slot().Padding(0).FillWidth(3.0f)
-				[
-					AssetSaveMode->AsShared()
-				]
 		];
 
-	//return Content;
-
 
 	TSharedPtr<SExpandableArea> AssetConfigPanel = SNew(SExpandableArea)
 		.HeaderPadding(FMargin(2.0f))
@@ -292,6 +339,7 @@ TSharedPtr<SWidget> FModelingToolsEditorModeToolkit::MakeAssetConfigPanel()
 
 
 
+
 void FModelingToolsEditorModeToolkit::UpdateActiveToolProperties()
 {
 	UInteractiveTool* CurTool = GetScriptableEditorMode()->GetToolManager()->GetActiveTool(EToolSide::Left);
@@ -576,6 +624,10 @@ void FModelingToolsEditorModeToolkit::OnToolStarted(UInteractiveToolManager* Man
 	ActiveToolName = CurTool->GetToolInfo().ToolDisplayName;
 
 	GetToolkitHost()->AddViewportOverlayWidget(ViewportOverlayWidget.ToSharedRef());
+
+	// disable LOD level picker once Tool is active
+	AssetLODMode->SetEnabled(false);
+	AssetLODModeLabel->SetEnabled(false);
 }
 
 void FModelingToolsEditorModeToolkit::OnToolEnded(UInteractiveToolManager* Manager, UInteractiveTool* Tool)
@@ -596,6 +648,10 @@ void FModelingToolsEditorModeToolkit::OnToolEnded(UInteractiveToolManager* Manag
 	{
 		CurTool->OnPropertySetsModified.RemoveAll(this);
 	}
+
+	// re-enable LOD level picker
+	AssetLODMode->SetEnabled(true);
+	AssetLODModeLabel->SetEnabled(true);
 }
 
 void FModelingToolsEditorModeToolkit::OnActiveViewportChanged(TSharedPtr<IAssetViewport> OldViewport, TSharedPtr<IAssetViewport> NewViewport)
diff --git a/Engine/Plugins/Experimental/ModelingToolsEditorMode/Source/ModelingToolsEditorMode/Public/ModelingToolsEditorModeToolkit.h b/Engine/Plugins/Experimental/ModelingToolsEditorMode/Source/ModelingToolsEditorMode/Public/ModelingToolsEditorModeToolkit.h
--- a/Engine/Plugins/Experimental/ModelingToolsEditorMode/Source/ModelingToolsEditorMode/Public/ModelingToolsEditorModeToolkit.h
+++ b/Engine/Plugins/Experimental/ModelingToolsEditorMode/Source/ModelingToolsEditorMode/Public/ModelingToolsEditorModeToolkit.h
@@ -97,6 +97,10 @@ private:
 	FDelegateHandle AssetSettingsModifiedHandle;
 	void OnShowAssetSettings();
 
+	TArray<TSharedPtr<FString>> AssetLODModes;
+	TSharedPtr<STextBlock> AssetLODModeLabel;
+	TSharedPtr<STextComboBox> AssetLODMode;
+
 
 	// custom accept/cancel/complete handlers
 };
diff --git a/Engine/Source/Editor/Experimental/EditorInteractiveToolsFramework/Private/EditorInteractiveToolsFrameworkModule.cpp b/Engine/Source/Editor/Experimental/EditorInteractiveToolsFramework/Private/EditorInteractiveToolsFrameworkModule.cpp
--- a/Engine/Source/Editor/Experimental/EditorInteractiveToolsFramework/Private/EditorInteractiveToolsFrameworkModule.cpp
+++ b/Engine/Source/Editor/Experimental/EditorInteractiveToolsFramework/Private/EditorInteractiveToolsFrameworkModule.cpp
@@ -8,10 +8,13 @@
 
 #define LOCTEXT_NAMESPACE "FEditorInteractiveToolsFrameworkModule"
 
+int32 FEditorInteractiveToolsFrameworkGlobals::RegisteredStaticMeshTargetFactoryKey = -1;
+
 void FEditorInteractiveToolsFrameworkModule::StartupModule()
 {
 	// This code will execute after your module is loaded into memory; the exact timing is specified in the .uplugin file per-module
-	AddComponentTargetFactory( TUniquePtr<FComponentTargetFactory>{new FStaticMeshComponentTargetFactory{} } );
+	FEditorInteractiveToolsFrameworkGlobals::RegisteredStaticMeshTargetFactoryKey 
+		= AddComponentTargetFactory( TUniquePtr<FComponentTargetFactory>{new FStaticMeshComponentTargetFactory{} } );
 
 	// The Interactive Tools Framework needs to be explicitly loaded, at the very least to make sure that
 	// the InteractiveToolsSelectionStoreSubsystem gets initialized.
diff --git a/Engine/Source/Editor/Experimental/EditorInteractiveToolsFramework/Public/EditorInteractiveToolsFrameworkModule.h b/Engine/Source/Editor/Experimental/EditorInteractiveToolsFramework/Public/EditorInteractiveToolsFrameworkModule.h
--- a/Engine/Source/Editor/Experimental/EditorInteractiveToolsFramework/Public/EditorInteractiveToolsFrameworkModule.h
+++ b/Engine/Source/Editor/Experimental/EditorInteractiveToolsFramework/Public/EditorInteractiveToolsFrameworkModule.h
@@ -5,6 +5,15 @@
 #include "CoreMinimal.h"
 #include "Modules/ModuleManager.h"
 
+
+class EDITORINTERACTIVETOOLSFRAMEWORK_API FEditorInteractiveToolsFrameworkGlobals
+{
+public:
+	// This is the key returned by AddComponentTargetFactory() for the FStaticMeshComponentTargetFactory created/registered
+	// in StartupModule() below. Use this key to find/remove that module registration if you need to.
+	static int32 RegisteredStaticMeshTargetFactoryKey;
+};
+
 class FEditorInteractiveToolsFrameworkModule : public IModuleInterface
 {
 public:
diff --git a/Engine/Source/Editor/UnrealEd/Private/Tools/EditorComponentSourceFactory.cpp b/Engine/Source/Editor/UnrealEd/Private/Tools/EditorComponentSourceFactory.cpp
--- a/Engine/Source/Editor/UnrealEd/Private/Tools/EditorComponentSourceFactory.cpp
+++ b/Engine/Source/Editor/UnrealEd/Private/Tools/EditorComponentSourceFactory.cpp
@@ -10,6 +10,55 @@
 #include "PhysicsEngine/BodySetup.h"
 
 
+static void DisplayCriticalWarningMessage(const FString& Message)
+{
+	if (GAreScreenMessagesEnabled)
+	{
+		GEngine->AddOnScreenDebugMessage(INDEX_NONE, 10.0f, FColor::Red, Message);
+	}
+	UE_LOG(LogTemp, Warning, TEXT("%s"), *Message);
+}
+
+
+FStaticMeshComponentTarget::FStaticMeshComponentTarget(UPrimitiveComponent* Component, EStaticMeshEditingLOD EditingLODIn)
+	: FPrimitiveComponentTarget(Cast<UStaticMeshComponent>(Component))
+{
+	EditingLOD = EStaticMeshEditingLOD::LOD0;
+
+	UStaticMeshComponent* StaticMeshComponent = Cast<UStaticMeshComponent>(Component);
+	if (ensure(StaticMeshComponent != nullptr))
+	{
+		UStaticMesh* StaticMeshAsset = StaticMeshComponent->GetStaticMesh();
+		if (ensure(StaticMeshAsset != nullptr))
+		{
+			if (EditingLODIn == EStaticMeshEditingLOD::MaxQuality)
+			{
+				EditingLOD = StaticMeshAsset->IsHiResMeshDescriptionValid() ? EStaticMeshEditingLOD::HiResSource : EStaticMeshEditingLOD::LOD0;
+			}
+			else if (EditingLODIn == EStaticMeshEditingLOD::HiResSource)
+			{
+				EditingLOD = StaticMeshAsset->IsHiResMeshDescriptionValid() ? EStaticMeshEditingLOD::HiResSource : EStaticMeshEditingLOD::LOD0;
+				if (EditingLOD != EStaticMeshEditingLOD::HiResSource)
+				{
+					DisplayCriticalWarningMessage(FString(TEXT("HiRes Source selected but not available - Falling Back to LOD0")));
+				}
+			}
+			else
+			{
+				int32 WantLOD = (int)EditingLODIn;
+				int32 MaxExistingLOD = StaticMeshAsset->GetNumSourceModels() - 1;
+				if (WantLOD > MaxExistingLOD)
+				{
+					DisplayCriticalWarningMessage(FString::Printf(TEXT("LOD%d Requested but not available - Falling Back to LOD%d"), WantLOD, MaxExistingLOD));
+					EditingLOD = (EStaticMeshEditingLOD)MaxExistingLOD;
+				}
+			}
+		}
+	}
+}
+
+
+
 
 bool FStaticMeshComponentTarget::IsValid() const
 {
@@ -27,17 +76,31 @@ bool FStaticMeshComponentTarget::IsValid() const
 	{
 		return false;
 	}
-	if (!(LODIndex < StaticMesh->GetNumSourceModels()))
+
+	if (EditingLOD == EStaticMeshEditingLOD::HiResSource)
+	{
+		if (StaticMesh->IsHiResMeshDescriptionValid() == false)
+		{
+			return false;
+		}
+	}
+	else if ( (int32)EditingLOD >= StaticMesh->GetNumSourceModels() )
 	{
 		return false;
 	}
+
 	return true;
 }
 
 FMeshDescription* FStaticMeshComponentTarget::GetMesh() 
 {
-	ensure(IsValid());
-	return IsValid() ? Cast<UStaticMeshComponent>(Component)->GetStaticMesh()->GetMeshDescription(LODIndex) : nullptr;
+	if (ensure(IsValid()))
+	{
+		UStaticMesh* StaticMesh = Cast<UStaticMeshComponent>(Component)->GetStaticMesh();
+		return (EditingLOD == EStaticMeshEditingLOD::HiResSource) ?
+			StaticMesh->GetHiResMeshDescription() : StaticMesh->GetMeshDescription((int32)EditingLOD);
+	}
+	return nullptr;
 }
 
 
@@ -149,12 +212,7 @@ void FStaticMeshComponentTarget::CommitMesh( const FCommitter& Committer )
 
 	if (StaticMesh->GetPathName().StartsWith(TEXT("/Engine/")))
 	{
-		FString DebugMessage = FString::Printf(TEXT("CANNOT MODIFY BUILT-IN ENGINE ASSET %s"), *StaticMesh->GetPathName());
-		if (GAreScreenMessagesEnabled)
-		{
-			GEngine->AddOnScreenDebugMessage(INDEX_NONE, 10.0f, FColor::Red, DebugMessage);
-		}
-		UE_LOG(LogTemp, Warning, TEXT("%s"), *DebugMessage);
+		DisplayCriticalWarningMessage(FString::Printf(TEXT("CANNOT MODIFY BUILT-IN ENGINE ASSET %s"), *StaticMesh->GetPathName()));
 		return;
 	}
 
@@ -171,11 +229,19 @@ void FStaticMeshComponentTarget::CommitMesh( const FCommitter& Committer )
 	check(bSavedToTransactionBuffer);
 
 	FCommitParams CommitParams;
-	CommitParams.MeshDescription = StaticMesh->GetMeshDescription(LODIndex);
+	CommitParams.MeshDescription = GetMesh();
 
 	Committer(CommitParams);
 
-	StaticMesh->CommitMeshDescription(LODIndex);
+	if (EditingLOD == EStaticMeshEditingLOD::HiResSource)
+	{
+		StaticMesh->CommitHiResMeshDescription();
+	}
+	else
+	{
+		StaticMesh->CommitMeshDescription( (int32)EditingLOD );
+	}
+
 	StaticMesh->PostEditChange();
 
 	// this rebuilds physics, but it doesn't undo!
@@ -205,7 +271,7 @@ TUniquePtr<FPrimitiveComponentTarget> FStaticMeshComponentTargetFactory::Build(U
 		&& StaticMeshComponent->GetStaticMesh() != nullptr 
 		&& StaticMeshComponent->GetStaticMesh()->GetNumSourceModels() > 0)
 	{
-		return TUniquePtr<FPrimitiveComponentTarget> { new FStaticMeshComponentTarget{Component} };
+		return MakeUnique<FStaticMeshComponentTarget>(Component, CurrentEditingLOD);
 	}
 	return {};
 }
diff --git a/Engine/Source/Editor/UnrealEd/Public/Tools/EditorComponentSourceFactory.h b/Engine/Source/Editor/UnrealEd/Public/Tools/EditorComponentSourceFactory.h
--- a/Engine/Source/Editor/UnrealEd/Public/Tools/EditorComponentSourceFactory.h
+++ b/Engine/Source/Editor/UnrealEd/Public/Tools/EditorComponentSourceFactory.h
@@ -7,19 +7,47 @@
 #include "Components/PrimitiveComponent.h"
 #include "Components/StaticMeshComponent.h"
 
+
+/**
+ * EStaticMeshEditingLOD specifies which LOD of a StaticMesh the FStaticMeshComponentTarget
+ * refers to. FStaticMeshComponentTarget will remap requested value to the closest available LOD.
+ */
+enum class EStaticMeshEditingLOD
+{
+	LOD0 = 0,
+	LOD1 = 1,
+	LOD2 = 2,
+	LOD3 = 3,
+	LOD4 = 4,
+	LOD5 = 5,
+	LOD6 = 6,
+	LOD7 = 7,
+
+	HiResSource = 20,	// HiRes source mesh is optional, and will fall back to LOD0 if not stored
+	MaxQuality = 21		// use HiRes source mesh if available, or LOD0 otherwise
+};
+
+
+
+
+
 class UNREALED_API FStaticMeshComponentTargetFactory : public FComponentTargetFactory
 {
 public:
+	// new FStaticMeshComponentTargets returned by Build() will be requested for this LOD
+	EStaticMeshEditingLOD CurrentEditingLOD = EStaticMeshEditingLOD::MaxQuality;
+
 	bool CanBuild( UActorComponent* Candidate ) override;
 	TUniquePtr<FPrimitiveComponentTarget> Build( UPrimitiveComponent* PrimitiveComponent ) override;
 };
 
+
+
 class UNREALED_API FStaticMeshComponentTarget : public FPrimitiveComponentTarget
 {
 public:
 
-	FStaticMeshComponentTarget( UPrimitiveComponent* Component )
-		: FPrimitiveComponentTarget( Cast<UStaticMeshComponent>(Component) ){}
+	FStaticMeshComponentTarget(UPrimitiveComponent* Component, EStaticMeshEditingLOD EditingLOD = EStaticMeshEditingLOD::LOD0);
 
 	virtual bool IsValid() const override;
 
@@ -33,5 +61,12 @@ public:
 
 	virtual bool HasSameSourceData(const FPrimitiveComponentTarget& OtherTarget) const override;
 
-	static const int LODIndex{0};
+
+protected:
+	// LOD to edit, default is to edit LOD0
+	EStaticMeshEditingLOD EditingLOD = EStaticMeshEditingLOD::LOD0;
 };
+
+
+
+
diff --git a/Engine/Source/Runtime/Experimental/InteractiveToolsFramework/Private/ComponentSourceInterfaces.cpp b/Engine/Source/Runtime/Experimental/InteractiveToolsFramework/Private/ComponentSourceInterfaces.cpp
--- a/Engine/Source/Runtime/Experimental/InteractiveToolsFramework/Private/ComponentSourceInterfaces.cpp
+++ b/Engine/Source/Runtime/Experimental/InteractiveToolsFramework/Private/ComponentSourceInterfaces.cpp
@@ -5,22 +5,46 @@
 #include "Components/PrimitiveComponent.h"
 #include "Containers/Array.h"
 
+
 namespace
 {
-	TArray<TUniquePtr<FComponentTargetFactory>> Factories;
+	TMap<uint32, TUniquePtr<FComponentTargetFactory>> Factories;
+
+	static int32 FactoryUniqueKeyGenerator = 1;
+}
+
+
+int32 AddComponentTargetFactory( TUniquePtr<FComponentTargetFactory> Factory )
+{
+	int32 NewFactoryKey = FactoryUniqueKeyGenerator++;
+	Factories.Add(NewFactoryKey, MoveTemp(Factory));
+	return NewFactoryKey;
+}
+
+
+FComponentTargetFactory* FindComponentTargetFactoryByKey(int32 Key)
+{
+	TUniquePtr<FComponentTargetFactory>* Found = Factories.Find(Key);
+	return (Found != nullptr) ? Found->Get() : nullptr;
 }
 
 
-void AddComponentTargetFactory( TUniquePtr<FComponentTargetFactory> Factory )
+bool RemoveComponentTargetFactoryByKey(int32 Key)
 {
-	Factories.Push( MoveTemp(Factory) );
+	return Factories.Remove(Key) != 0;
 }
 
+void RemoveAllComponentTargetFactoryies()
+{
+	Factories.Reset();
+}
+
+
 bool CanMakeComponentTarget(UActorComponent* Component)
 {
-	for ( const auto& Factory : Factories )
+	for ( const auto& FactoryPair : Factories )
 	{
-		if ( Factory->CanBuild(Component) )
+		if (FactoryPair.Value->CanBuild(Component) )
 		{
 			return true;
 		}
@@ -30,11 +54,11 @@ bool CanMakeComponentTarget(UActorComponent* Component)
 
 TUniquePtr<FPrimitiveComponentTarget> MakeComponentTarget(UPrimitiveComponent* Component)
 {
-	for ( const auto& Factory : Factories )
+	for (const auto& FactoryPair : Factories)
 	{
-		if ( Factory->CanBuild( Component ) )
+		if (FactoryPair.Value->CanBuild( Component ) )
 		{
-			return Factory->Build( Component );
+			return FactoryPair.Value->Build( Component );
 		}
 	}
 	return {};
diff --git a/Engine/Source/Runtime/Experimental/InteractiveToolsFramework/Public/ComponentSourceInterfaces.h b/Engine/Source/Runtime/Experimental/InteractiveToolsFramework/Public/ComponentSourceInterfaces.h
--- a/Engine/Source/Runtime/Experimental/InteractiveToolsFramework/Public/ComponentSourceInterfaces.h
+++ b/Engine/Source/Runtime/Experimental/InteractiveToolsFramework/Public/ComponentSourceInterfaces.h
@@ -123,9 +123,27 @@ public:
  *
  * Add a factory method to make ComponentTarget from UPrimitiveComponent*
  * @param Factory The ComponentTargetFactory
- * @return void
+ * @return integer indentifier that identifies this Factory
  */
-INTERACTIVETOOLSFRAMEWORK_API void AddComponentTargetFactory( TUniquePtr<FComponentTargetFactory> Factory );
+INTERACTIVETOOLSFRAMEWORK_API int32 AddComponentTargetFactory( TUniquePtr<FComponentTargetFactory> Factory );
+
+
+/**
+ * @return ComponentTargetFactory associated with the given Key, or nullptr if not found
+ */
+INTERACTIVETOOLSFRAMEWORK_API FComponentTargetFactory* FindComponentTargetFactoryByKey(int32 Key);
+
+
+/**
+ * @return true if ComponentTargetFactory associated with the given Key was found and removed
+ */
+INTERACTIVETOOLSFRAMEWORK_API bool RemoveComponentTargetFactoryByKey(int32 Key);
+
+
+/**
+ * remove all registered ComponentTargetFactory objects
+ */
+INTERACTIVETOOLSFRAMEWORK_API void RemoveAllComponentTargetFactoryies();
 
 
 /**
