ModelingMode UX: fix issue with disappearing mode panel, by allocating the widget before calling FModeToolkit::Init(). See code comments for an explanation of why this is necessary.
#rnx

diff --git a/Engine/Plugins/Experimental/ModelingToolsEditorMode/Source/ModelingToolsEditorMode/Private/ModelingToolsEditorModeToolkit.cpp b/Engine/Plugins/Experimental/ModelingToolsEditorMode/Source/ModelingToolsEditorMode/Private/ModelingToolsEditorModeToolkit.cpp
--- a/Engine/Plugins/Experimental/ModelingToolsEditorMode/Source/ModelingToolsEditorMode/Private/ModelingToolsEditorModeToolkit.cpp
+++ b/Engine/Plugins/Experimental/ModelingToolsEditorMode/Source/ModelingToolsEditorMode/Private/ModelingToolsEditorModeToolkit.cpp
@@ -54,6 +54,18 @@ FModelingToolsEditorModeToolkit::~FModelingToolsEditorModeToolkit()
 
 void FModelingToolsEditorModeToolkit::Init(const TSharedPtr<IToolkitHost>& InitToolkitHost, TWeakObjectPtr<UEdMode> InOwningMode)
 {
+	// Have to create the ToolkitWidget here because FModeToolkit::Init() is going to ask for it and add
+	// it to the Mode panel, and not ask again afterwards. However we have to call Init() to get the 
+	// ModeDetailsView created, that we need to add to the ToolkitWidget. So, we will create the Widget
+	// here but only add the rows to it after we call Init()
+	TSharedPtr<SVerticalBox> ToolkitWidgetVBox = SNew(SVerticalBox);
+	SAssignNew(ToolkitWidget, SBorder)
+		.HAlign(HAlign_Fill)
+		.Padding(4)
+		[
+			ToolkitWidgetVBox->AsShared()
+		];
+
 	FModeToolkit::Init(InitToolkitHost, InOwningMode);
 
 	GetToolkitHost()->OnActiveViewportChanged().AddSP(this, &FModelingToolsEditorModeToolkit::OnActiveViewportChanged);
@@ -78,37 +90,26 @@ void FModelingToolsEditorModeToolkit::Init(const TSharedPtr<IToolkitHost>& InitT
 		.ColorAndOpacity(FSlateColor(FLinearColor(0.9f, 0.15f, 0.15f)));
 	ToolWarningArea->SetText(FText::GetEmpty());
 
-
-	SAssignNew(ToolkitWidget, SBorder)
-		.HAlign(HAlign_Fill)
-		.Padding(4)
+	// add the various sections to the mode toolbox
+	ToolkitWidgetVBox->AddSlot().AutoHeight().HAlign(HAlign_Fill).Padding(5)
 		[
-			SNew(SVerticalBox)
-
-			+ SVerticalBox::Slot().AutoHeight().HAlign(HAlign_Fill).Padding(5)
-				[
-					ModeWarningArea->AsShared()
-				]
-
-			+ SVerticalBox::Slot().AutoHeight().HAlign(HAlign_Fill).Padding(5)
-				[
-					ModeHeaderArea->AsShared()
-				]
-
-			+ SVerticalBox::Slot().AutoHeight().HAlign(HAlign_Fill).AutoHeight().Padding(5)
-				[
-					ToolWarningArea->AsShared()
-				]
-
-			+ SVerticalBox::Slot().HAlign(HAlign_Fill).FillHeight(1.f)
-				[
-					ModeDetailsView->AsShared()
-				]
-
-			+ SVerticalBox::Slot().AutoHeight().HAlign(HAlign_Fill).Padding(5)
-				[
-					MakeAssetConfigPanel()->AsShared()
-				]
+			ModeWarningArea->AsShared()
+		];
+	ToolkitWidgetVBox->AddSlot().AutoHeight().HAlign(HAlign_Fill).Padding(5)
+		[
+			ModeHeaderArea->AsShared()
+		];
+	ToolkitWidgetVBox->AddSlot().AutoHeight().HAlign(HAlign_Fill).Padding(5)
+		[
+			ToolWarningArea->AsShared()
+		];
+	ToolkitWidgetVBox->AddSlot().HAlign(HAlign_Fill).FillHeight(1.f)
+		[
+			ModeDetailsView->AsShared()
+		];
+	ToolkitWidgetVBox->AddSlot().AutoHeight().HAlign(HAlign_Fill).Padding(5)
+		[
+			MakeAssetConfigPanel()->AsShared()
 		];
 
 	ClearNotification();
