ModelingMode UX: add missing icons
#rnx

diff --git a/Engine/Plugins/Experimental/ModelingToolsEditorMode/Source/ModelingToolsEditorMode/Private/ModelingToolsEditorModeStyle.cpp b/Engine/Plugins/Experimental/ModelingToolsEditorMode/Source/ModelingToolsEditorMode/Private/ModelingToolsEditorModeStyle.cpp
--- a/Engine/Plugins/Experimental/ModelingToolsEditorMode/Source/ModelingToolsEditorMode/Private/ModelingToolsEditorModeStyle.cpp
+++ b/Engine/Plugins/Experimental/ModelingToolsEditorMode/Source/ModelingToolsEditorMode/Private/ModelingToolsEditorModeStyle.cpp
@@ -103,22 +103,22 @@ void FModelingToolsEditorModeStyle::Initialize()
 
 
 		// Modes Palette Toolbar Icons
-		StyleSet->Set("ModelingToolsManagerCommands.BeginAddBoxPrimitiveTool", 			new IMAGE_PLUGIN_BRUSH("Icons/ModelingBox_x20", 		Icon20x20));
-		StyleSet->Set("ModelingToolsManagerCommands.BeginAddBoxPrimitiveTool.Small", 		new IMAGE_PLUGIN_BRUSH("Icons/ModelingBox_x40",		Icon40x40));
-		StyleSet->Set("ModelingToolsManagerCommands.BeginAddCylinderPrimitiveTool", 			new IMAGE_PLUGIN_BRUSH("Icons/ModelingCylinder_x20", 		Icon20x20));
-		StyleSet->Set("ModelingToolsManagerCommands.BeginAddCylinderPrimitiveTool.Small", 		new IMAGE_PLUGIN_BRUSH("Icons/ModelingCylinder_x40",		Icon40x40));
-		StyleSet->Set("ModelingToolsManagerCommands.BeginAddConePrimitiveTool", 			new IMAGE_PLUGIN_BRUSH("Icons/ModelingCone_x20", 		Icon20x20));
-		StyleSet->Set("ModelingToolsManagerCommands.BeginAddConePrimitiveTool.Small", 		new IMAGE_PLUGIN_BRUSH("Icons/ModelingCone_x40",		Icon40x40));
-		StyleSet->Set("ModelingToolsManagerCommands.BeginAddArrowPrimitiveTool", 			new IMAGE_PLUGIN_BRUSH("Icons/ModelingArrow_x20", 		Icon20x20));
-		StyleSet->Set("ModelingToolsManagerCommands.BeginAddArrowPrimitiveTool.Small", 		new IMAGE_PLUGIN_BRUSH("Icons/ModelingArrow_x40",		Icon40x40));
+		StyleSet->Set("ModelingToolsManagerCommands.BeginAddBoxPrimitiveTool", 			new IMAGE_BRUSH_SVG("Icons/ModelingBox", 		Icon20x20));
+		StyleSet->Set("ModelingToolsManagerCommands.BeginAddBoxPrimitiveTool.Small", 		new IMAGE_BRUSH_SVG("Icons/ModelingBox",		Icon40x40));
+		StyleSet->Set("ModelingToolsManagerCommands.BeginAddCylinderPrimitiveTool", 			new IMAGE_BRUSH_SVG("Icons/ModelingCylinder", 		Icon20x20));
+		StyleSet->Set("ModelingToolsManagerCommands.BeginAddCylinderPrimitiveTool.Small", 		new IMAGE_BRUSH_SVG("Icons/ModelingCylinder",		Icon40x40));
+		StyleSet->Set("ModelingToolsManagerCommands.BeginAddConePrimitiveTool", 			new IMAGE_BRUSH_SVG("Icons/ModelingCone", 		Icon20x20));
+		StyleSet->Set("ModelingToolsManagerCommands.BeginAddConePrimitiveTool.Small", 		new IMAGE_BRUSH_SVG("Icons/ModelingCone",		Icon40x40));
+		StyleSet->Set("ModelingToolsManagerCommands.BeginAddArrowPrimitiveTool", 			new IMAGE_BRUSH_SVG("Icons/ModelingArrow", 		Icon20x20));
+		StyleSet->Set("ModelingToolsManagerCommands.BeginAddArrowPrimitiveTool.Small", 		new IMAGE_BRUSH_SVG("Icons/ModelingArrow",		Icon40x40));
 		StyleSet->Set("ModelingToolsManagerCommands.BeginAddRectanglePrimitiveTool", 			new IMAGE_PLUGIN_BRUSH("Icons/ModelingRectangle_x20", 		Icon20x20));
 		StyleSet->Set("ModelingToolsManagerCommands.BeginAddRectanglePrimitiveTool.Small", 		new IMAGE_PLUGIN_BRUSH("Icons/ModelingRectangle_x40",		Icon40x40));
 		StyleSet->Set("ModelingToolsManagerCommands.BeginAddDiscPrimitiveTool", 			new IMAGE_PLUGIN_BRUSH("Icons/ModelingDisc_x20", 		Icon20x20));
 		StyleSet->Set("ModelingToolsManagerCommands.BeginAddDiscPrimitiveTool.Small", 		new IMAGE_PLUGIN_BRUSH("Icons/ModelingDisc_x40",		Icon40x40));
-		StyleSet->Set("ModelingToolsManagerCommands.BeginAddTorusPrimitiveTool", 			new IMAGE_PLUGIN_BRUSH("Icons/ModelingTorus_x20", 		Icon20x20));
-		StyleSet->Set("ModelingToolsManagerCommands.BeginAddTorusPrimitiveTool.Small", 		new IMAGE_PLUGIN_BRUSH("Icons/ModelingTorus_x40",		Icon40x40));
-		StyleSet->Set("ModelingToolsManagerCommands.BeginAddSpherePrimitiveTool", 			new IMAGE_PLUGIN_BRUSH("Icons/ModelingSphere_x20", 		Icon20x20));
-		StyleSet->Set("ModelingToolsManagerCommands.BeginAddSpherePrimitiveTool.Small", 		new IMAGE_PLUGIN_BRUSH("Icons/ModelingSphere_x40",		Icon40x40));
+		StyleSet->Set("ModelingToolsManagerCommands.BeginAddTorusPrimitiveTool", 			new IMAGE_BRUSH_SVG("Icons/ModelingTorus", 		Icon20x20));
+		StyleSet->Set("ModelingToolsManagerCommands.BeginAddTorusPrimitiveTool.Small", 		new IMAGE_BRUSH_SVG("Icons/ModelingTorus",		Icon40x40));
+		StyleSet->Set("ModelingToolsManagerCommands.BeginAddSpherePrimitiveTool", 			new IMAGE_BRUSH_SVG("Icons/ModelingSphere", 		Icon20x20));
+		StyleSet->Set("ModelingToolsManagerCommands.BeginAddSpherePrimitiveTool.Small", 		new IMAGE_BRUSH_SVG("Icons/ModelingSphere",		Icon40x40));
 
 		StyleSet->Set("ModelingToolsManagerCommands.BeginDrawPolygonTool", 				new IMAGE_PLUGIN_BRUSH("Icons/DrawPolygon_40x",		Icon20x20));
 		StyleSet->Set("ModelingToolsManagerCommands.BeginDrawPolygonTool.Small", 		new IMAGE_PLUGIN_BRUSH("Icons/DrawPolygon_40x", 	Icon20x20));
@@ -204,8 +204,8 @@ void FModelingToolsEditorModeStyle::Initialize()
 		StyleSet->Set("ModelingToolsManagerCommands.BeginGlobalUVGenerateTool.Small",        new IMAGE_PLUGIN_BRUSH("Icons/AutoUnwrap_40x",       Icon20x20));      
 		StyleSet->Set("ModelingToolsManagerCommands.BeginBakeTransformTool",                 new IMAGE_PLUGIN_BRUSH("Icons/BakeXForm_40x",        Icon20x20));      
 		StyleSet->Set("ModelingToolsManagerCommands.BeginBakeTransformTool.Small",           new IMAGE_PLUGIN_BRUSH("Icons/BakeXForm_40x",        Icon20x20));      
-		StyleSet->Set("ModelingToolsManagerCommands.BeginCombineMeshesTool",                 new IMAGE_PLUGIN_BRUSH("Icons/Combine_40x",          Icon20x20));   
-		StyleSet->Set("ModelingToolsManagerCommands.BeginCombineMeshesTool.Small",           new IMAGE_PLUGIN_BRUSH("Icons/Combine_40x",          Icon20x20));   
+		StyleSet->Set("ModelingToolsManagerCommands.BeginCombineMeshesTool",                 new IMAGE_BRUSH_SVG("Icons/GeometryCombine",          Icon20x20));
+		StyleSet->Set("ModelingToolsManagerCommands.BeginCombineMeshesTool.Small",           new IMAGE_BRUSH_SVG("Icons/GeometryCombine",          Icon20x20));
 		StyleSet->Set("ModelingToolsManagerCommands.BeginDuplicateMeshesTool",               new IMAGE_PLUGIN_BRUSH("Icons/Duplicate_40x",        Icon20x20));   
 		StyleSet->Set("ModelingToolsManagerCommands.BeginDuplicateMeshesTool.Small",         new IMAGE_PLUGIN_BRUSH("Icons/Duplicate_40x",        Icon20x20));   
 		StyleSet->Set("ModelingToolsManagerCommands.BeginEditMeshMaterialsTool",             new IMAGE_PLUGIN_BRUSH("Icons/EditMats_40x",         Icon20x20));     
@@ -248,6 +248,12 @@ void FModelingToolsEditorModeStyle::Initialize()
 		StyleSet->Set("ModelingToolsManagerCommands.BeginTransformUVIslandsTool.Small",      new IMAGE_PLUGIN_BRUSH("Icons/TransformUVs_40x",     Icon20x20));         
 		StyleSet->Set("ModelingToolsManagerCommands.BeginUVLayoutTool",                      new IMAGE_PLUGIN_BRUSH("Icons/UVLayout_40x",         Icon20x20));    
 		StyleSet->Set("ModelingToolsManagerCommands.BeginUVLayoutTool.Small",                new IMAGE_PLUGIN_BRUSH("Icons/UVLayout_40x",         Icon20x20));    
+		StyleSet->Set("ModelingToolsManagerCommands.BeginMeshGroupPaintTool", new IMAGE_BRUSH_SVG("Icons/GroupPaint", Icon20x20));
+		StyleSet->Set("ModelingToolsManagerCommands.BeginMeshGroupPaintTool.Small", new IMAGE_BRUSH_SVG("Icons/GroupPaint", Icon20x20));
+		StyleSet->Set("ModelingToolsManagerCommands.BeginLatticeDeformerTool", new IMAGE_BRUSH_SVG("Icons/LatticeDeformation", Icon20x20));
+		StyleSet->Set("ModelingToolsManagerCommands.BeginLatticeDeformerTool.Small", new IMAGE_BRUSH_SVG("Icons/LatticeDeformation", Icon20x20));
+
+		
 
 		StyleSet->Set("ModelingToolsManagerCommands.BeginVolumeToMeshTool",                  new IMAGE_PLUGIN_BRUSH("Icons/ModelingVol2Mesh_x40",         Icon20x20));
 		StyleSet->Set("ModelingToolsManagerCommands.BeginVolumeToMeshTool.Small",            new IMAGE_PLUGIN_BRUSH("Icons/ModelingVol2Mesh_x40",         Icon20x20));
@@ -262,27 +268,17 @@ void FModelingToolsEditorModeStyle::Initialize()
 		StyleSet->Set("ModelingToolsManagerCommands.BeginExtractCollisionGeometryTool",      new IMAGE_PLUGIN_BRUSH("Icons/ModelingCollisionToMesh_x40",         Icon20x20));
 		StyleSet->Set("ModelingToolsManagerCommands.BeginExtractCollisionGeometryTool.Small",new IMAGE_PLUGIN_BRUSH("Icons/ModelingCollisionToMesh_x40",         Icon20x20));
 
-		//const FLinearColor LayerSelectionColor = FLinearColor(0.13f, 0.70f, 1.00f);
-
-		//// Selection color for the active row should be ??? to align with the editor viewport.
-		//const FTableRowStyle& NormalTableRowStyle = FEditorStyle::Get().GetWidgetStyle<FTableRowStyle>("TableView.Row");
-
-		//StyleSet->Set("TileMapEditor.LayerBrowser.TableViewRow",
-		//	FTableRowStyle(NormalTableRowStyle)
-		//	.SetActiveBrush(IMAGE_BRUSH("Common/Selection", Icon8x8, LayerSelectionColor))
-		//	.SetActiveHoveredBrush(IMAGE_BRUSH("Common/Selection", Icon8x8, LayerSelectionColor))
-		//	.SetInactiveBrush(IMAGE_BRUSH("Common/Selection", Icon8x8, LayerSelectionColor))
-		//	.SetInactiveHoveredBrush(IMAGE_BRUSH("Common/Selection", Icon8x8, LayerSelectionColor))
-		//);
+		StyleSet->Set("ModelingToolsManagerCommands.BeginGenerateStaticMeshLODAssetTool", new IMAGE_BRUSH_SVG("Icons/AutoLOD", Icon20x20));
+		StyleSet->Set("ModelingToolsManagerCommands.BeginGenerateStaticMeshLODAssetTool.Small", new IMAGE_BRUSH_SVG("Icons/AutoLOD", Icon20x20));
 
-		//StyleSet->Set("TileMapEditor.LayerBrowser.SelectionColor", LayerSelectionColor);
+		StyleSet->Set("ModelingToolsManagerCommands.BeginGroomCardsEditorTool", new IMAGE_BRUSH_SVG("Icons/CardsEditor", Icon20x20));
+		StyleSet->Set("ModelingToolsManagerCommands.BeginGroomCardsEditorTool.Small", new IMAGE_BRUSH_SVG("Icons/CardsEditor", Icon20x20));
+		StyleSet->Set("ModelingToolsManagerCommands.BeginGenerateLODMeshesTool", new IMAGE_BRUSH_SVG("Icons/GenLODs", Icon20x20));
+		StyleSet->Set("ModelingToolsManagerCommands.BeginGenerateLODMeshesTool.Small", new IMAGE_BRUSH_SVG("Icons/GenLODs", Icon20x20));
+		StyleSet->Set("ModelingToolsManagerCommands.BeginGroomToMeshTool", new IMAGE_BRUSH_SVG("Icons/HairHelmet", Icon20x20));
+		StyleSet->Set("ModelingToolsManagerCommands.BeginGroomToMeshTool.Small", new IMAGE_BRUSH_SVG("Icons/HairHelmet", Icon20x20));
 
-		//StyleSet->Set("TileMapEditor.TileSetPalette.NothingSelectedText", FTextBlockStyle(NormalText)
-		//	.SetFont(DEFAULT_FONT("BoldCondensed", 18))
-		//	.SetColorAndOpacity(FLinearColor(0.8, 0.8f, 0.0f, 0.8f))
-		//	.SetShadowOffset(FVector2D(1.0f, 1.0f))
-		//	.SetShadowColorAndOpacity(FLinearColor(0.0f, 0.0f, 0.0f, 0.9f))
-		//);
+		
 	}
 
 	FSlateStyleRegistry::RegisterSlateStyle(*StyleSet.Get());
diff --git a/Engine/Plugins/Experimental/ModelingToolsEditorMode/Source/ModelingToolsEditorMode/Private/ModelingToolsManagerActions.cpp b/Engine/Plugins/Experimental/ModelingToolsEditorMode/Source/ModelingToolsEditorMode/Private/ModelingToolsManagerActions.cpp
--- a/Engine/Plugins/Experimental/ModelingToolsEditorMode/Source/ModelingToolsEditorMode/Private/ModelingToolsManagerActions.cpp
+++ b/Engine/Plugins/Experimental/ModelingToolsEditorMode/Source/ModelingToolsEditorMode/Private/ModelingToolsManagerActions.cpp
@@ -79,7 +79,7 @@ void FModelingToolsManagerCommands::RegisterCommands()
 	UI_COMMAND(BeginPlaneCutTool, "PlnCut", "Cut Selected Meshes with Plane", EUserInterfaceActionType::ToggleButton, FInputChord());
 	UI_COMMAND(BeginMirrorTool, "Mirror", "Mirror Selected Meshes", EUserInterfaceActionType::ToggleButton, FInputChord());
 	UI_COMMAND(BeginHoleFillTool, "HFill", "Fill Holes in Mesh", EUserInterfaceActionType::ToggleButton, FInputChord());
-	UI_COMMAND(BeginLatticeDeformerTool, "Grid", "Deform Mesh with 3D Lattice/Grid", EUserInterfaceActionType::ToggleButton, FInputChord());
+	UI_COMMAND(BeginLatticeDeformerTool, "Lattice", "Deform Mesh with 3D Lattice/Grid", EUserInterfaceActionType::ToggleButton, FInputChord());
 	UI_COMMAND(BeginPolygonCutTool, "PolyCut", "Cut Mesh with Extruded Polygon", EUserInterfaceActionType::ToggleButton, FInputChord());
 	UI_COMMAND(BeginGlobalUVGenerateTool, "AutoUV", "Auto-Unwrap and Pack UVs for Mesh", EUserInterfaceActionType::ToggleButton, FInputChord());
 	UI_COMMAND(BeginGroupUVGenerateTool, "Unwrap", "Recalculate UV unwrapping for Mesh regions", EUserInterfaceActionType::ToggleButton, FInputChord());
