ModelingMode UX: add display of active tool icon to Viewport Overlay, fix up missing/incorrect tool display names
#rnx

diff --git a/Engine/Plugins/Experimental/MeshLODToolset/Source/MeshLODToolset/Private/Tools/GenerateStaticMeshLODAssetTool.cpp b/Engine/Plugins/Experimental/MeshLODToolset/Source/MeshLODToolset/Private/Tools/GenerateStaticMeshLODAssetTool.cpp
--- a/Engine/Plugins/Experimental/MeshLODToolset/Source/MeshLODToolset/Private/Tools/GenerateStaticMeshLODAssetTool.cpp
+++ b/Engine/Plugins/Experimental/MeshLODToolset/Source/MeshLODToolset/Private/Tools/GenerateStaticMeshLODAssetTool.cpp
@@ -196,7 +196,7 @@ void UGenerateStaticMeshLODAssetTool::Setup()
 
 	UInteractiveTool::Setup();
 
-	SetToolDisplayName(LOCTEXT("ToolName", "Generate LOD"));
+	SetToolDisplayName(LOCTEXT("ToolName", "AutoLOD"));
 	GetToolManager()->DisplayMessage(
 		LOCTEXT("OnStartStaticMeshLODAssetTool", "Create a new LOD asset"),
 		EToolMessageLevel::UserNotification);
diff --git a/Engine/Plugins/Experimental/MeshModelingToolset/Source/MeshModelingTools/Private/LatticeDeformerTool.cpp b/Engine/Plugins/Experimental/MeshModelingToolset/Source/MeshModelingTools/Private/LatticeDeformerTool.cpp
--- a/Engine/Plugins/Experimental/MeshModelingToolset/Source/MeshModelingTools/Private/LatticeDeformerTool.cpp
+++ b/Engine/Plugins/Experimental/MeshModelingToolset/Source/MeshModelingTools/Private/LatticeDeformerTool.cpp
@@ -99,7 +99,7 @@ void ULatticeDeformerTool::Setup()
 {
 	UInteractiveTool::Setup();
 
-	SetToolDisplayName(LOCTEXT("ToolName", "Grid Deform"));
+	SetToolDisplayName(LOCTEXT("ToolName", "Lattice Deform"));
 	GetToolManager()->DisplayMessage(LOCTEXT("LatticeDeformerToolMessage", 
 		"Drag the lattice control points to deform the mesh"), EToolMessageLevel::UserNotification);
 
diff --git a/Engine/Plugins/Experimental/MeshModelingToolset/Source/MeshModelingToolsEditorOnly/Private/GenerateLODMeshesTool.cpp b/Engine/Plugins/Experimental/MeshModelingToolset/Source/MeshModelingToolsEditorOnly/Private/GenerateLODMeshesTool.cpp
--- a/Engine/Plugins/Experimental/MeshModelingToolset/Source/MeshModelingToolsEditorOnly/Private/GenerateLODMeshesTool.cpp
+++ b/Engine/Plugins/Experimental/MeshModelingToolset/Source/MeshModelingToolsEditorOnly/Private/GenerateLODMeshesTool.cpp
@@ -99,6 +99,8 @@ void UGenerateLODMeshesTool::Setup()
 {
 	UInteractiveTool::Setup();
 
+	SetToolDisplayName(LOCTEXT("ToolName", "Generate HairHelmet LODs"));
+
 	{
 		// if in editor, create progress indicator dialog because building mesh copies can be slow (for very large meshes)
 		// this is especially needed because of the copy we make of the meshdescription; for Reasons, copying meshdescription is pretty slow
diff --git a/Engine/Plugins/Experimental/ModelingToolsEditorMode/Source/ModelingToolsEditorMode/Private/ModelingToolsEditorMode.cpp b/Engine/Plugins/Experimental/ModelingToolsEditorMode/Source/ModelingToolsEditorMode/Private/ModelingToolsEditorMode.cpp
--- a/Engine/Plugins/Experimental/ModelingToolsEditorMode/Source/ModelingToolsEditorMode/Private/ModelingToolsEditorMode.cpp
+++ b/Engine/Plugins/Experimental/ModelingToolsEditorMode/Source/ModelingToolsEditorMode/Private/ModelingToolsEditorMode.cpp
@@ -347,27 +347,27 @@ void UModelingToolsEditorMode::Enter()
 	//
 	auto AddPatchToolBuilder = NewObject<UAddPatchToolBuilder>();
 	AddPatchToolBuilder->AssetAPI = ModelingModeAssetGenerationAPI.Get();
-	RegisterTool(ToolManagerCommands.BeginAddPatchTool, TEXT("AddPatchTool"), AddPatchToolBuilder);
+	RegisterTool(ToolManagerCommands.BeginAddPatchTool, TEXT("BeginAddPatchTool"), AddPatchToolBuilder);
 
 	auto RevolveBoundaryToolBuilder = NewObject<URevolveBoundaryToolBuilder>();
 	RevolveBoundaryToolBuilder->AssetAPI = ModelingModeAssetGenerationAPI.Get();
-	RegisterTool(ToolManagerCommands.BeginRevolveBoundaryTool, TEXT("RevolveBoundaryTool"), RevolveBoundaryToolBuilder);
+	RegisterTool(ToolManagerCommands.BeginRevolveBoundaryTool, TEXT("BeginRevolveBoundaryTool"), RevolveBoundaryToolBuilder);
 
 	auto DrawPolygonToolBuilder = NewObject<UDrawPolygonToolBuilder>();
 	DrawPolygonToolBuilder->AssetAPI = ModelingModeAssetGenerationAPI.Get();
-	RegisterTool(ToolManagerCommands.BeginDrawPolygonTool, TEXT("DrawPolygonTool"), DrawPolygonToolBuilder);
+	RegisterTool(ToolManagerCommands.BeginDrawPolygonTool, TEXT("BeginDrawPolygonTool"), DrawPolygonToolBuilder);
 
 	auto DrawPolyPathToolBuilder = NewObject<UDrawPolyPathToolBuilder>();
 	DrawPolyPathToolBuilder->AssetAPI = ModelingModeAssetGenerationAPI.Get();
-	RegisterTool(ToolManagerCommands.BeginDrawPolyPathTool, TEXT("DrawPolyPath"), DrawPolyPathToolBuilder);
+	RegisterTool(ToolManagerCommands.BeginDrawPolyPathTool, TEXT("BeginDrawPolyPathTool"), DrawPolyPathToolBuilder);
 
 	auto DrawAndRevolveToolBuilder = NewObject<UDrawAndRevolveToolBuilder>();
 	DrawAndRevolveToolBuilder->AssetAPI = ModelingModeAssetGenerationAPI.Get();
-	RegisterTool(ToolManagerCommands.BeginDrawAndRevolveTool, TEXT("RevolveTool"), DrawAndRevolveToolBuilder);
+	RegisterTool(ToolManagerCommands.BeginDrawAndRevolveTool, TEXT("BeginDrawAndRevolveTool"), DrawAndRevolveToolBuilder);
 
 	auto ShapeSprayToolBuilder = NewObject<UShapeSprayToolBuilder>();
 	ShapeSprayToolBuilder->AssetAPI = ModelingModeAssetGenerationAPI.Get();
-	RegisterTool(ToolManagerCommands.BeginShapeSprayTool, TEXT("ShapeSprayTool"), ShapeSprayToolBuilder);
+	RegisterTool(ToolManagerCommands.BeginShapeSprayTool, TEXT("BeginShapeSprayTool"), ShapeSprayToolBuilder);
 
 
 	//
@@ -376,46 +376,46 @@ void UModelingToolsEditorMode::Enter()
 
 	auto MoveVerticesToolBuilder = NewObject<UMeshVertexSculptToolBuilder>();
 	MoveVerticesToolBuilder->StylusAPI = StylusStateTracker.Get();
-	RegisterTool(ToolManagerCommands.BeginSculptMeshTool, TEXT("MoveVerticesTool"), MoveVerticesToolBuilder);
+	RegisterTool(ToolManagerCommands.BeginSculptMeshTool, TEXT("BeginSculptMeshTool"), MoveVerticesToolBuilder);
 
 	auto MeshGroupPaintToolBuilder = NewObject<UMeshGroupPaintToolBuilder>();
 	MeshGroupPaintToolBuilder->StylusAPI = StylusStateTracker.Get();
-	RegisterTool(ToolManagerCommands.BeginMeshGroupPaintTool, TEXT("MeshGroupPaintTool"), MeshGroupPaintToolBuilder);
+	RegisterTool(ToolManagerCommands.BeginMeshGroupPaintTool, TEXT("BeginMeshGroupPaintTool"), MeshGroupPaintToolBuilder);
 
-	RegisterTool(ToolManagerCommands.BeginPolyEditTool, TEXT("EditMeshPolygonsTool"), NewObject<UEditMeshPolygonsToolBuilder>());
+	RegisterTool(ToolManagerCommands.BeginPolyEditTool, TEXT("BeginPolyEditTool"), NewObject<UEditMeshPolygonsToolBuilder>());
 	UEditMeshPolygonsToolBuilder* TriEditBuilder = NewObject<UEditMeshPolygonsToolBuilder>();
 	TriEditBuilder->bTriangleMode = true;
-	RegisterTool(ToolManagerCommands.BeginTriEditTool, TEXT("EditMeshTrianglesTool"), TriEditBuilder);
-	RegisterTool(ToolManagerCommands.BeginPolyDeformTool, TEXT("DeformMeshPolygonsTool"), NewObject<UDeformMeshPolygonsToolBuilder>());
-	RegisterTool(ToolManagerCommands.BeginSmoothMeshTool, TEXT("SmoothMeshTool"), NewObject<USmoothMeshToolBuilder>());
-	RegisterTool(ToolManagerCommands.BeginOffsetMeshTool, TEXT("OffsetMeshTool"), NewObject<UOffsetMeshToolBuilder>());
-	RegisterTool(ToolManagerCommands.BeginDisplaceMeshTool, TEXT("DisplaceMeshTool"), NewObject<UDisplaceMeshToolBuilder>());
-	RegisterTool(ToolManagerCommands.BeginMeshSpaceDeformerTool, TEXT("MeshSpaceDeformerTool"), NewObject<UMeshSpaceDeformerToolBuilder>());
-	RegisterTool(ToolManagerCommands.BeginTransformMeshesTool, TEXT("TransformMeshesTool"), NewObject<UTransformMeshesToolBuilder>());
-	RegisterTool(ToolManagerCommands.BeginEditPivotTool, TEXT("EditPivotTool"), NewObject<UEditPivotToolBuilder>());
-	RegisterTool(ToolManagerCommands.BeginAlignObjectsTool, TEXT("AlignObjects"), NewObject<UAlignObjectsToolBuilder>());
-	RegisterTool(ToolManagerCommands.BeginBakeTransformTool, TEXT("BakeTransformTool"), NewObject<UBakeTransformToolBuilder>());
-	RegisterTool(ToolManagerCommands.BeginTransformUVIslandsTool, TEXT("EditUVIslands"), NewObject<UEditUVIslandsToolBuilder>());
-	RegisterTool(ToolManagerCommands.BeginLatticeDeformerTool, TEXT("LatticeDeformerTool"), NewObject<ULatticeDeformerToolBuilder>());
-	RegisterTool(ToolManagerCommands.BeginSubdividePolyTool, TEXT("SubdividePolyTool"), NewObject<USubdividePolyToolBuilder>());
+	RegisterTool(ToolManagerCommands.BeginTriEditTool, TEXT("BeginTriEditTool"), TriEditBuilder);
+	RegisterTool(ToolManagerCommands.BeginPolyDeformTool, TEXT("BeginPolyDeformTool"), NewObject<UDeformMeshPolygonsToolBuilder>());
+	RegisterTool(ToolManagerCommands.BeginSmoothMeshTool, TEXT("BeginSmoothMeshTool"), NewObject<USmoothMeshToolBuilder>());
+	RegisterTool(ToolManagerCommands.BeginOffsetMeshTool, TEXT("BeginOffsetMeshTool"), NewObject<UOffsetMeshToolBuilder>());
+	RegisterTool(ToolManagerCommands.BeginDisplaceMeshTool, TEXT("BeginDisplaceMeshTool"), NewObject<UDisplaceMeshToolBuilder>());
+	RegisterTool(ToolManagerCommands.BeginMeshSpaceDeformerTool, TEXT("BeginMeshSpaceDeformerTool"), NewObject<UMeshSpaceDeformerToolBuilder>());
+	RegisterTool(ToolManagerCommands.BeginTransformMeshesTool, TEXT("BeginTransformMeshesTool"), NewObject<UTransformMeshesToolBuilder>());
+	RegisterTool(ToolManagerCommands.BeginEditPivotTool, TEXT("BeginEditPivotTool"), NewObject<UEditPivotToolBuilder>());
+	RegisterTool(ToolManagerCommands.BeginAlignObjectsTool, TEXT("BeginAlignObjectsTool"), NewObject<UAlignObjectsToolBuilder>());
+	RegisterTool(ToolManagerCommands.BeginBakeTransformTool, TEXT("BeginBakeTransformTool"), NewObject<UBakeTransformToolBuilder>());
+	RegisterTool(ToolManagerCommands.BeginTransformUVIslandsTool, TEXT("BeginTransformUVIslandsTool"), NewObject<UEditUVIslandsToolBuilder>());
+	RegisterTool(ToolManagerCommands.BeginLatticeDeformerTool, TEXT("BeginLatticeDeformerTool"), NewObject<ULatticeDeformerToolBuilder>());
+	RegisterTool(ToolManagerCommands.BeginSubdividePolyTool, TEXT("BeginSubdividePolyTool"), NewObject<USubdividePolyToolBuilder>());
 
 	UCombineMeshesToolBuilder* CombineMeshesToolBuilder = NewObject<UCombineMeshesToolBuilder>();
 	CombineMeshesToolBuilder->AssetAPI = ModelingModeAssetGenerationAPI.Get();
-	RegisterTool(ToolManagerCommands.BeginCombineMeshesTool, TEXT("CombineMeshesTool"), CombineMeshesToolBuilder);
+	RegisterTool(ToolManagerCommands.BeginCombineMeshesTool, TEXT("BeginCombineMeshesTool"), CombineMeshesToolBuilder);
 
 	UCombineMeshesToolBuilder* DuplicateMeshesToolBuilder = NewObject<UCombineMeshesToolBuilder>();
 	DuplicateMeshesToolBuilder->AssetAPI = ModelingModeAssetGenerationAPI.Get();
 	DuplicateMeshesToolBuilder->bIsDuplicateTool = true;
-	RegisterTool(ToolManagerCommands.BeginDuplicateMeshesTool, TEXT("DuplicateMeshesTool"), DuplicateMeshesToolBuilder);
+	RegisterTool(ToolManagerCommands.BeginDuplicateMeshesTool, TEXT("BeginDuplicateMeshesTool"), DuplicateMeshesToolBuilder);
 
 
 	ULODManagerToolBuilder* LODManagerToolBuilder = NewObject<ULODManagerToolBuilder>();
 	LODManagerToolBuilder->AssetAPI = ModelingModeAssetGenerationAPI.Get();
-	RegisterTool(ToolManagerCommands.BeginLODManagerTool, TEXT("LODManagerTool"), LODManagerToolBuilder);
+	RegisterTool(ToolManagerCommands.BeginLODManagerTool, TEXT("BeginLODManagerTool"), LODManagerToolBuilder);
 
 	UGenerateStaticMeshLODAssetToolBuilder* GenerateSMLODToolBuilder = NewObject<UGenerateStaticMeshLODAssetToolBuilder>();
 	GenerateSMLODToolBuilder->AssetAPI = ModelingModeAssetGenerationAPI.Get();
-	RegisterTool(ToolManagerCommands.BeginGenerateStaticMeshLODAssetTool, TEXT("GenerateSMLODTool"), GenerateSMLODToolBuilder);
+	RegisterTool(ToolManagerCommands.BeginGenerateStaticMeshLODAssetTool, TEXT("BeginGenerateStaticMeshLODAssetTool"), GenerateSMLODToolBuilder);
 
 
 	// edit tools
@@ -424,141 +424,141 @@ void UModelingToolsEditorMode::Enter()
 	auto DynaSculptToolBuilder = NewObject<UDynamicMeshSculptToolBuilder>();
 	DynaSculptToolBuilder->bEnableRemeshing = true;
 	DynaSculptToolBuilder->StylusAPI = StylusStateTracker.Get();
-	RegisterTool(ToolManagerCommands.BeginRemeshSculptMeshTool, TEXT("DynaSculptTool"), DynaSculptToolBuilder);
+	RegisterTool(ToolManagerCommands.BeginRemeshSculptMeshTool, TEXT("BeginRemeshSculptMeshTool"), DynaSculptToolBuilder);
 
-	RegisterTool(ToolManagerCommands.BeginRemeshMeshTool, TEXT("RemeshMeshTool"), NewObject<URemeshMeshToolBuilder>());
-	RegisterTool(ToolManagerCommands.BeginProjectToTargetTool, TEXT("ProjectToTargetTool"), NewObject<UProjectToTargetToolBuilder>());
-	RegisterTool(ToolManagerCommands.BeginSimplifyMeshTool, TEXT("SimplifyMeshTool"), NewObject<USimplifyMeshToolBuilder>());
+	RegisterTool(ToolManagerCommands.BeginRemeshMeshTool, TEXT("BeginRemeshMeshTool"), NewObject<URemeshMeshToolBuilder>());
+	RegisterTool(ToolManagerCommands.BeginProjectToTargetTool, TEXT("BeginProjectToTargetTool"), NewObject<UProjectToTargetToolBuilder>());
+	RegisterTool(ToolManagerCommands.BeginSimplifyMeshTool, TEXT("BeginSimplifyMeshTool"), NewObject<USimplifyMeshToolBuilder>());
 
 	auto GroupEdgeInsertionToolBuilder = NewObject<UGroupEdgeInsertionToolBuilder>();
 	GroupEdgeInsertionToolBuilder->AssetAPI = ToolsContext->GetAssetAPI();
-	RegisterTool(ToolManagerCommands.BeginGroupEdgeInsertionTool, TEXT("GroupEdgeInsertionTool"), GroupEdgeInsertionToolBuilder);
+	RegisterTool(ToolManagerCommands.BeginGroupEdgeInsertionTool, TEXT("BeginGroupEdgeInsertionTool"), GroupEdgeInsertionToolBuilder);
 
 	auto EdgeLoopInsertionToolBuilder = NewObject<UEdgeLoopInsertionToolBuilder>();
 	EdgeLoopInsertionToolBuilder->AssetAPI = ToolsContext->GetAssetAPI();
-	RegisterTool(ToolManagerCommands.BeginEdgeLoopInsertionTool, TEXT("EdgeLoopInsertionTool"), EdgeLoopInsertionToolBuilder);
+	RegisterTool(ToolManagerCommands.BeginEdgeLoopInsertionTool, TEXT("BeginEdgeLoopInsertionTool"), EdgeLoopInsertionToolBuilder);
 
 	auto EditNormalsToolBuilder = NewObject<UEditNormalsToolBuilder>();
 	EditNormalsToolBuilder->AssetAPI = ModelingModeAssetGenerationAPI.Get();
-	RegisterTool(ToolManagerCommands.BeginEditNormalsTool, TEXT("EditNormalsTool"), EditNormalsToolBuilder);
+	RegisterTool(ToolManagerCommands.BeginEditNormalsTool, TEXT("BeginEditNormalsTool"), EditNormalsToolBuilder);
 
 	auto TangentsToolBuilder = NewObject<UMeshTangentsToolBuilder>();
-	RegisterTool(ToolManagerCommands.BeginEditTangentsTool, TEXT("MeshTangentsTool"), TangentsToolBuilder);
+	RegisterTool(ToolManagerCommands.BeginEditTangentsTool, TEXT("BeginEditTangentsTool"), TangentsToolBuilder);
 
 	auto RemoveOccludedTrianglesToolBuilder = NewObject<URemoveOccludedTrianglesToolBuilder>();
 	RemoveOccludedTrianglesToolBuilder->AssetAPI = ModelingModeAssetGenerationAPI.Get();
-	RegisterTool(ToolManagerCommands.BeginRemoveOccludedTrianglesTool, TEXT("RemoveOccludedTrianglesTool"), RemoveOccludedTrianglesToolBuilder);
+	RegisterTool(ToolManagerCommands.BeginRemoveOccludedTrianglesTool, TEXT("BeginRemoveOccludedTrianglesTool"), RemoveOccludedTrianglesToolBuilder);
 
 	auto HoleFillToolBuilder = NewObject<UHoleFillToolBuilder>();
-	RegisterTool(ToolManagerCommands.BeginHoleFillTool, TEXT("HoleFillTool"), HoleFillToolBuilder);
+	RegisterTool(ToolManagerCommands.BeginHoleFillTool, TEXT("BeginHoleFillTool"), HoleFillToolBuilder);
 
 	auto UVProjectionToolBuilder = NewObject<UUVProjectionToolBuilder>();
 	UVProjectionToolBuilder->AssetAPI = ModelingModeAssetGenerationAPI.Get();
-	RegisterTool(ToolManagerCommands.BeginUVProjectionTool, TEXT("UVProjectionTool"), UVProjectionToolBuilder);
+	RegisterTool(ToolManagerCommands.BeginUVProjectionTool, TEXT("BeginUVProjectionTool"), UVProjectionToolBuilder);
 
 	auto UVLayoutToolBuilder = NewObject<UUVLayoutToolBuilder>();
 	UVLayoutToolBuilder->AssetAPI = ModelingModeAssetGenerationAPI.Get();
-	RegisterTool(ToolManagerCommands.BeginUVLayoutTool, TEXT("UVLayoutTool"), UVLayoutToolBuilder);
+	RegisterTool(ToolManagerCommands.BeginUVLayoutTool, TEXT("BeginUVLayoutTool"), UVLayoutToolBuilder);
 
 	auto MergeMeshesToolBuilder = NewObject<UMergeMeshesToolBuilder>();
 	MergeMeshesToolBuilder->AssetAPI = ModelingModeAssetGenerationAPI.Get();
-	RegisterTool(ToolManagerCommands.BeginVoxelMergeTool, TEXT("MergeMeshesTool"), MergeMeshesToolBuilder);
+	RegisterTool(ToolManagerCommands.BeginVoxelMergeTool, TEXT("BeginVoxelMergeTool"), MergeMeshesToolBuilder);
 
 	auto VoxelCSGMeshesToolBuilder = NewObject<UVoxelCSGMeshesToolBuilder>();
 	VoxelCSGMeshesToolBuilder->AssetAPI = ModelingModeAssetGenerationAPI.Get();
-	RegisterTool(ToolManagerCommands.BeginVoxelBooleanTool, TEXT("VoxelCSGMeshesTool"), VoxelCSGMeshesToolBuilder);
+	RegisterTool(ToolManagerCommands.BeginVoxelBooleanTool, TEXT("BeginVoxelBooleanTool"), VoxelCSGMeshesToolBuilder);
 
 	auto VoxelSolidifyMeshesToolBuilder = NewObject<UVoxelSolidifyMeshesToolBuilder>();
 	VoxelSolidifyMeshesToolBuilder->AssetAPI = ModelingModeAssetGenerationAPI.Get();
-	RegisterTool(ToolManagerCommands.BeginVoxelSolidifyTool, TEXT("VoxelSolidifyMeshesTool"), VoxelSolidifyMeshesToolBuilder);
+	RegisterTool(ToolManagerCommands.BeginVoxelSolidifyTool, TEXT("BeginVoxelSolidifyTool"), VoxelSolidifyMeshesToolBuilder);
 
 	auto VoxelBlendMeshesToolBuilder = NewObject<UVoxelBlendMeshesToolBuilder>();
 	VoxelBlendMeshesToolBuilder->AssetAPI = ModelingModeAssetGenerationAPI.Get();
-	RegisterTool(ToolManagerCommands.BeginVoxelBlendTool, TEXT("VoxelBlendMeshesTool"), VoxelBlendMeshesToolBuilder);
+	RegisterTool(ToolManagerCommands.BeginVoxelBlendTool, TEXT("BeginVoxelBlendTool"), VoxelBlendMeshesToolBuilder);
 
 	auto VoxelMorphologyMeshesToolBuilder = NewObject<UVoxelMorphologyMeshesToolBuilder>();
 	VoxelMorphologyMeshesToolBuilder->AssetAPI = ModelingModeAssetGenerationAPI.Get();
-	RegisterTool(ToolManagerCommands.BeginVoxelMorphologyTool, TEXT("VoxelMorphologyMeshesTool"), VoxelMorphologyMeshesToolBuilder);
+	RegisterTool(ToolManagerCommands.BeginVoxelMorphologyTool, TEXT("BeginVoxelMorphologyTool"), VoxelMorphologyMeshesToolBuilder);
 
 	auto SelfUnionMeshesToolBuilder = NewObject<USelfUnionMeshesToolBuilder>();
 	SelfUnionMeshesToolBuilder->AssetAPI = ModelingModeAssetGenerationAPI.Get();
-	RegisterTool(ToolManagerCommands.BeginSelfUnionTool, TEXT("SelfUnionMeshesTool"), SelfUnionMeshesToolBuilder);
+	RegisterTool(ToolManagerCommands.BeginSelfUnionTool, TEXT("BeginSelfUnionTool"), SelfUnionMeshesToolBuilder);
 
 	auto CSGMeshesToolBuilder = NewObject<UCSGMeshesToolBuilder>();
 	CSGMeshesToolBuilder->AssetAPI = ModelingModeAssetGenerationAPI.Get();
-	RegisterTool(ToolManagerCommands.BeginMeshBooleanTool, TEXT("CSGMeshesTool"), CSGMeshesToolBuilder);
+	RegisterTool(ToolManagerCommands.BeginMeshBooleanTool, TEXT("BeginMeshBooleanTool"), CSGMeshesToolBuilder);
 
 	auto TrimMeshesToolBuilder = NewObject<UCSGMeshesToolBuilder>();
 	TrimMeshesToolBuilder->bTrimMode = true;
 	TrimMeshesToolBuilder->AssetAPI = ModelingModeAssetGenerationAPI.Get();
-	RegisterTool(ToolManagerCommands.BeginMeshTrimTool, TEXT("TrimMeshesTool"), TrimMeshesToolBuilder);
+	RegisterTool(ToolManagerCommands.BeginMeshTrimTool, TEXT("BeginMeshTrimTool"), TrimMeshesToolBuilder);
 
 	auto BspConversionToolBuilder = NewObject<UBspConversionToolBuilder>();
 	BspConversionToolBuilder->AssetAPI = ModelingModeAssetGenerationAPI.Get();
-	RegisterTool(ToolManagerCommands.BeginBspConversionTool, TEXT("BspConversionTool"), BspConversionToolBuilder);
+	RegisterTool(ToolManagerCommands.BeginBspConversionTool, TEXT("BeginBspConversionTool"), BspConversionToolBuilder);
 
 	auto MeshToVolumeToolBuilder = NewObject<UMeshToVolumeToolBuilder>();
 	MeshToVolumeToolBuilder->AssetAPI = ModelingModeAssetGenerationAPI.Get();
-	RegisterTool(ToolManagerCommands.BeginMeshToVolumeTool, TEXT("MeshToVolumeTool"), MeshToVolumeToolBuilder);
+	RegisterTool(ToolManagerCommands.BeginMeshToVolumeTool, TEXT("BeginMeshToVolumeTool"), MeshToVolumeToolBuilder);
 
 	auto VolumeToMeshToolBuilder = NewObject<UVolumeToMeshToolBuilder>();
 	VolumeToMeshToolBuilder->AssetAPI = ModelingModeAssetGenerationAPI.Get();
-	RegisterTool(ToolManagerCommands.BeginVolumeToMeshTool, TEXT("VolumeToMeshTool"), VolumeToMeshToolBuilder);
+	RegisterTool(ToolManagerCommands.BeginVolumeToMeshTool, TEXT("BeginVolumeToMeshTool"), VolumeToMeshToolBuilder);
 
 	auto PlaneCutToolBuilder = NewObject<UPlaneCutToolBuilder>();
 	PlaneCutToolBuilder->AssetAPI = ModelingModeAssetGenerationAPI.Get();
-	RegisterTool(ToolManagerCommands.BeginPlaneCutTool, TEXT("PlaneCutTool"), PlaneCutToolBuilder);
+	RegisterTool(ToolManagerCommands.BeginPlaneCutTool, TEXT("BeginPlaneCutTool"), PlaneCutToolBuilder);
 
 	auto MirrorToolBuilder = NewObject<UMirrorToolBuilder>();
 	MirrorToolBuilder->AssetAPI = ModelingModeAssetGenerationAPI.Get();
-	RegisterTool(ToolManagerCommands.BeginMirrorTool, TEXT("MirrorTool"), MirrorToolBuilder);
+	RegisterTool(ToolManagerCommands.BeginMirrorTool, TEXT("BeginMirrorTool"), MirrorToolBuilder);
 
 	auto PolygonCutToolBuilder = NewObject<UPolygonOnMeshToolBuilder>();
-	RegisterTool(ToolManagerCommands.BeginPolygonCutTool, TEXT("PolyCutTool"), PolygonCutToolBuilder);
+	RegisterTool(ToolManagerCommands.BeginPolygonCutTool, TEXT("BeginPolygonCutTool"), PolygonCutToolBuilder);
 
 	auto GlobalUVGenerateToolBuilder = NewObject<UParameterizeMeshToolBuilder>();
 	GlobalUVGenerateToolBuilder->AssetAPI = ModelingModeAssetGenerationAPI.Get();
 	GlobalUVGenerateToolBuilder->bDoAutomaticGlobalUnwrap = true;
-	RegisterTool(ToolManagerCommands.BeginGlobalUVGenerateTool, TEXT("GlobalParameterizeMeshTool"), GlobalUVGenerateToolBuilder);
+	RegisterTool(ToolManagerCommands.BeginGlobalUVGenerateTool, TEXT("BeginGlobalUVGenerateTool"), GlobalUVGenerateToolBuilder);
 
 	auto GroupUVGenerateToolBuilder = NewObject<UParameterizeMeshToolBuilder>();
 	GroupUVGenerateToolBuilder->AssetAPI = ModelingModeAssetGenerationAPI.Get();
 	GroupUVGenerateToolBuilder->bDoAutomaticGlobalUnwrap = false;
-	RegisterTool(ToolManagerCommands.BeginGroupUVGenerateTool, TEXT("GroupParameterizeMeshTool"), GroupUVGenerateToolBuilder);
+	RegisterTool(ToolManagerCommands.BeginGroupUVGenerateTool, TEXT("BeginGroupUVGenerateTool"), GroupUVGenerateToolBuilder);
 
-	RegisterTool(ToolManagerCommands.BeginUVSeamEditTool, TEXT("UVSeamSculptTool"), NewObject< USeamSculptToolBuilder>());
+	RegisterTool(ToolManagerCommands.BeginUVSeamEditTool, TEXT("BeginUVSeamEditTool"), NewObject< USeamSculptToolBuilder>());
 
 	auto MeshSelectionToolBuilder = NewObject<UMeshSelectionToolBuilder>();
 	MeshSelectionToolBuilder->AssetAPI = ModelingModeAssetGenerationAPI.Get();
-	RegisterTool(ToolManagerCommands.BeginMeshSelectionTool, TEXT("MeshSelectionTool"), MeshSelectionToolBuilder);
+	RegisterTool(ToolManagerCommands.BeginMeshSelectionTool, TEXT("BeginMeshSelectionTool"), MeshSelectionToolBuilder);
 
 	auto EditMeshMaterialsToolBuilder = NewObject<UEditMeshMaterialsToolBuilder>();
 	EditMeshMaterialsToolBuilder->AssetAPI = ModelingModeAssetGenerationAPI.Get();
-	RegisterTool(ToolManagerCommands.BeginEditMeshMaterialsTool, TEXT("EditMaterialsTool"), EditMeshMaterialsToolBuilder);
+	RegisterTool(ToolManagerCommands.BeginEditMeshMaterialsTool, TEXT("BeginEditMeshMaterialsTool"), EditMeshMaterialsToolBuilder);
 	
-	RegisterTool(ToolManagerCommands.BeginMeshAttributePaintTool, TEXT("MeshAttributePaintTool"), NewObject<UMeshAttributePaintToolBuilder>());
+	RegisterTool(ToolManagerCommands.BeginMeshAttributePaintTool, TEXT("BeginMeshAttributePaintTool"), NewObject<UMeshAttributePaintToolBuilder>());
 
 	auto BakeMeshAttributeMapsToolBuilder = NewObject<UBakeMeshAttributeMapsToolBuilder>();
 	BakeMeshAttributeMapsToolBuilder->AssetAPI = ModelingModeAssetGenerationAPI.Get();
-	RegisterTool(ToolManagerCommands.BeginBakeMeshAttributeMapsTool, TEXT("BakeMeshMapsTool"), BakeMeshAttributeMapsToolBuilder);
+	RegisterTool(ToolManagerCommands.BeginBakeMeshAttributeMapsTool, TEXT("BeginBakeMeshAttributeMapsTool"), BakeMeshAttributeMapsToolBuilder);
 
 	// analysis tools
 
-	RegisterTool(ToolManagerCommands.BeginMeshInspectorTool, TEXT("MeshInspectorTool"), NewObject<UMeshInspectorToolBuilder>());
-	RegisterTool(ToolManagerCommands.BeginWeldEdgesTool, TEXT("WeldMeshEdgesTool"), NewObject<UWeldMeshEdgesToolBuilder>());
-	RegisterTool(ToolManagerCommands.BeginPolyGroupsTool, TEXT("ConvertToPolygonsTool"), NewObject<UConvertToPolygonsToolBuilder>());
-	RegisterTool(ToolManagerCommands.BeginAttributeEditorTool, TEXT("AttributeEditorTool"), NewObject<UAttributeEditorToolBuilder>());
+	RegisterTool(ToolManagerCommands.BeginMeshInspectorTool, TEXT("BeginMeshInspectorTool"), NewObject<UMeshInspectorToolBuilder>());
+	RegisterTool(ToolManagerCommands.BeginWeldEdgesTool, TEXT("BeginWeldEdgesTool"), NewObject<UWeldMeshEdgesToolBuilder>());
+	RegisterTool(ToolManagerCommands.BeginPolyGroupsTool, TEXT("BeginPolyGroupsTool"), NewObject<UConvertToPolygonsToolBuilder>());
+	RegisterTool(ToolManagerCommands.BeginAttributeEditorTool, TEXT("BeginAttributeEditorTool"), NewObject<UAttributeEditorToolBuilder>());
 
 
 	// Physics Tools
 
-	RegisterTool(ToolManagerCommands.BeginPhysicsInspectorTool, TEXT("PhysicsInspectorTool"), NewObject<UPhysicsInspectorToolBuilder>());
-	RegisterTool(ToolManagerCommands.BeginSetCollisionGeometryTool, TEXT("SetCollisionGeoTool"), NewObject<USetCollisionGeometryToolBuilder>());
+	RegisterTool(ToolManagerCommands.BeginPhysicsInspectorTool, TEXT("BeginPhysicsInspectorTool"), NewObject<UPhysicsInspectorToolBuilder>());
+	RegisterTool(ToolManagerCommands.BeginSetCollisionGeometryTool, TEXT("BeginSetCollisionGeometryTool"), NewObject<USetCollisionGeometryToolBuilder>());
 	//RegisterTool(ToolManagerCommands.BeginEditCollisionGeometryTool, TEXT("EditCollisionGeoTool"), NewObject<UEditCollisionGeometryToolBuilder>());
 
 	auto ExtractCollisionGeoToolBuilder = NewObject<UExtractCollisionGeometryToolBuilder>();
 	ExtractCollisionGeoToolBuilder->AssetAPI = ModelingModeAssetGenerationAPI.Get();
-	RegisterTool(ToolManagerCommands.BeginExtractCollisionGeometryTool, TEXT("ExtractCollisionGeoTool"), ExtractCollisionGeoToolBuilder);
+	RegisterTool(ToolManagerCommands.BeginExtractCollisionGeometryTool, TEXT("BeginExtractCollisionGeometryTool"), ExtractCollisionGeoToolBuilder);
 
 
 
@@ -566,13 +566,13 @@ void UModelingToolsEditorMode::Enter()
 
 	UGroomToMeshToolBuilder* GroomToMeshToolBuilder = NewObject<UGroomToMeshToolBuilder>();
 	GroomToMeshToolBuilder->AssetAPI = ModelingModeAssetGenerationAPI.Get();
-	RegisterTool(ToolManagerCommands.BeginGroomToMeshTool, TEXT("GroomToMeshTool"), GroomToMeshToolBuilder);
+	RegisterTool(ToolManagerCommands.BeginGroomToMeshTool, TEXT("BeginGroomToMeshTool"), GroomToMeshToolBuilder);
 
-	RegisterTool(ToolManagerCommands.BeginGroomCardsEditorTool, TEXT("GroomCardsEditorTool"), NewObject<UGroomCardsEditorToolBuilder>());
+	RegisterTool(ToolManagerCommands.BeginGroomCardsEditorTool, TEXT("BeginGroomCardsEditorTool"), NewObject<UGroomCardsEditorToolBuilder>());
 
 	UGenerateLODMeshesToolBuilder* GenerateLODMeshesToolBuilder = NewObject<UGenerateLODMeshesToolBuilder>();
 	GenerateLODMeshesToolBuilder->AssetAPI = ModelingModeAssetGenerationAPI.Get();
-	RegisterTool(ToolManagerCommands.BeginGenerateLODMeshesTool, TEXT("GenerateLODMeshesTool"), GenerateLODMeshesToolBuilder);
+	RegisterTool(ToolManagerCommands.BeginGenerateLODMeshesTool, TEXT("BeginGenerateLODMeshesTool"), GenerateLODMeshesToolBuilder);
 
 	// PolyModeling tools
 	auto RegisterPolyModelSelectTool = [&](EEditMeshPolygonsToolSelectionMode SelectionMode, TSharedPtr<FUICommandInfo> UICommand, FString StringName)
diff --git a/Engine/Plugins/Experimental/ModelingToolsEditorMode/Source/ModelingToolsEditorMode/Private/ModelingToolsEditorModeToolkit.cpp b/Engine/Plugins/Experimental/ModelingToolsEditorMode/Source/ModelingToolsEditorMode/Private/ModelingToolsEditorModeToolkit.cpp
--- a/Engine/Plugins/Experimental/ModelingToolsEditorMode/Source/ModelingToolsEditorMode/Private/ModelingToolsEditorModeToolkit.cpp
+++ b/Engine/Plugins/Experimental/ModelingToolsEditorMode/Source/ModelingToolsEditorMode/Private/ModelingToolsEditorModeToolkit.cpp
@@ -4,6 +4,7 @@
 #include "ModelingToolsEditorMode.h"
 #include "ModelingToolsManagerActions.h"
 #include "ModelingToolsEditorModeSettings.h"
+#include "ModelingToolsEditorModeStyle.h"
 #include "Engine/Selection.h"
 
 #include "Framework/MultiBox/MultiBoxBuilder.h"
@@ -134,6 +135,15 @@ void FModelingToolsEditorModeToolkit::Init(const TSharedPtr<IToolkitHost>& InitT
 		[
 			SNew(SHorizontalBox)
 
+			+SHorizontalBox::Slot()
+			.AutoWidth()
+			.VAlign(VAlign_Center)
+			.Padding(FMargin(0.f, 0.f, 8.f, 0.f))
+			[
+				SNew(SImage)
+				.Image_Lambda([this] () { return ActiveToolIcon; })
+			]
+
 			+SHorizontalBox::Slot()
 			.AutoWidth()
 			.VAlign(VAlign_Center)
@@ -627,6 +637,13 @@ void FModelingToolsEditorModeToolkit::OnToolStarted(UInteractiveToolManager* Man
 	ModeHeaderArea->SetVisibility(EVisibility::Collapsed);
 	ActiveToolName = CurTool->GetToolInfo().ToolDisplayName;
 
+	// try to update icon
+	FString ActiveToolIdentifier = GetScriptableEditorMode()->GetToolManager()->GetActiveToolName(EToolSide::Left);
+	ActiveToolIdentifier.InsertAt(0, ".");
+	FName ActiveToolIconName = ISlateStyle::Join(FModelingToolsManagerCommands::Get().GetContextName(), TCHAR_TO_ANSI(*ActiveToolIdentifier));
+	ActiveToolIcon = FModelingToolsEditorModeStyle::Get()->GetOptionalBrush(ActiveToolIconName);
+
+
 	GetToolkitHost()->AddViewportOverlayWidget(ViewportOverlayWidget.ToSharedRef());
 
 	// disable LOD level picker once Tool is active
diff --git a/Engine/Plugins/Experimental/ModelingToolsEditorMode/Source/ModelingToolsEditorMode/Public/ModelingToolsEditorModeToolkit.h b/Engine/Plugins/Experimental/ModelingToolsEditorMode/Source/ModelingToolsEditorMode/Public/ModelingToolsEditorModeToolkit.h
--- a/Engine/Plugins/Experimental/ModelingToolsEditorMode/Source/ModelingToolsEditorMode/Public/ModelingToolsEditorModeToolkit.h
+++ b/Engine/Plugins/Experimental/ModelingToolsEditorMode/Source/ModelingToolsEditorMode/Public/ModelingToolsEditorModeToolkit.h
@@ -68,6 +68,7 @@ private:
 	FText ActiveToolName;
 	FText ActiveToolMessage;
 	FStatusBarMessageHandle ActiveToolMessageHandle;
+	const FSlateBrush* ActiveToolIcon;
 
 	TSharedPtr<SWidget> ToolkitWidget;
 	void UpdateActiveToolProperties();
