ModelingTools: Add icon to Stairs tool

#rnx

diff --git a/Engine/Plugins/Experimental/ModelingToolsEditorMode/Source/ModelingToolsEditorMode/Private/ModelingToolsEditorModeStyle.cpp b/Engine/Plugins/Experimental/ModelingToolsEditorMode/Source/ModelingToolsEditorMode/Private/ModelingToolsEditorModeStyle.cpp
--- a/Engine/Plugins/Experimental/ModelingToolsEditorMode/Source/ModelingToolsEditorMode/Private/ModelingToolsEditorModeStyle.cpp
+++ b/Engine/Plugins/Experimental/ModelingToolsEditorMode/Source/ModelingToolsEditorMode/Private/ModelingToolsEditorModeStyle.cpp
@@ -119,6 +119,8 @@ void FModelingToolsEditorModeStyle::Initialize()
 		StyleSet->Set("ModelingToolsManagerCommands.BeginAddTorusPrimitiveTool.Small", 		new IMAGE_BRUSH_SVG("Icons/ModelingTorus",		Icon40x40));
 		StyleSet->Set("ModelingToolsManagerCommands.BeginAddSpherePrimitiveTool", 			new IMAGE_BRUSH_SVG("Icons/ModelingSphere", 		Icon20x20));
 		StyleSet->Set("ModelingToolsManagerCommands.BeginAddSpherePrimitiveTool.Small", 		new IMAGE_BRUSH_SVG("Icons/ModelingSphere",		Icon40x40));
+		StyleSet->Set("ModelingToolsManagerCommands.BeginAddStairsPrimitiveTool", 			new IMAGE_BRUSH_SVG("Icons/Staircase", Icon20x20));
+		StyleSet->Set("ModelingToolsManagerCommands.BeginAddStairsPrimitiveTool.Small",		new IMAGE_BRUSH_SVG("Icons/Staircase", Icon40x40));
 
 		StyleSet->Set("ModelingToolsManagerCommands.BeginDrawPolygonTool", 				new IMAGE_PLUGIN_BRUSH("Icons/DrawPolygon_40x",		Icon20x20));
 		StyleSet->Set("ModelingToolsManagerCommands.BeginDrawPolygonTool.Small", 		new IMAGE_PLUGIN_BRUSH("Icons/DrawPolygon_40x", 	Icon20x20));
