ModelingMode UX: add missing tool icons
#rnx

diff --git a/Engine/Plugins/Experimental/ModelingToolsEditorMode/Source/ModelingToolsEditorMode/Private/ModelingToolsEditorModeStyle.cpp b/Engine/Plugins/Experimental/ModelingToolsEditorMode/Source/ModelingToolsEditorMode/Private/ModelingToolsEditorModeStyle.cpp
--- a/Engine/Plugins/Experimental/ModelingToolsEditorMode/Source/ModelingToolsEditorMode/Private/ModelingToolsEditorModeStyle.cpp
+++ b/Engine/Plugins/Experimental/ModelingToolsEditorMode/Source/ModelingToolsEditorMode/Private/ModelingToolsEditorModeStyle.cpp
@@ -133,8 +133,8 @@ void FModelingToolsEditorModeStyle::Initialize()
 		StyleSet->Set("ModelingToolsManagerCommands.BeginSculptMeshTool.Small", 		new IMAGE_PLUGIN_BRUSH("Icons/Sculpt_40x", 			Icon20x20));
 		StyleSet->Set("ModelingToolsManagerCommands.BeginPolyEditTool", 				new IMAGE_PLUGIN_BRUSH("Icons/PolyEdit_40x", 		Icon20x20));
 		StyleSet->Set("ModelingToolsManagerCommands.BeginPolyEditTool.Small", 			new IMAGE_PLUGIN_BRUSH("Icons/PolyEdit_40x", 		Icon20x20));
-		//StyleSet->Set("ModelingToolsManagerCommands.BeginSubdividePolyTool",			new IMAGE_PLUGIN_BRUSH("Icons/PolyEdit_40x",		Icon20x20));
-		//StyleSet->Set("ModelingToolsManagerCommands.BeginSubdividePolyTool.Small",		new IMAGE_PLUGIN_BRUSH("Icons/PolyEdit_40x",		Icon20x20));
+		StyleSet->Set("ModelingToolsManagerCommands.BeginSubdividePolyTool",			new IMAGE_BRUSH_SVG("Icons/ModelingSubD",		Icon20x20));
+		StyleSet->Set("ModelingToolsManagerCommands.BeginSubdividePolyTool.Small",		new IMAGE_BRUSH_SVG("Icons/ModelingSubD",		Icon20x20));
 		StyleSet->Set("ModelingToolsManagerCommands.BeginGroupEdgeInsertionTool", 		new IMAGE_PLUGIN_BRUSH("Icons/ModelingGroupEdgeInsert_x40", 		Icon20x20));
 		StyleSet->Set("ModelingToolsManagerCommands.BeginGroupEdgeInsertionTool.Small", new IMAGE_PLUGIN_BRUSH("Icons/ModelingGroupEdgeInsert_x40", 		Icon20x20));
 		StyleSet->Set("ModelingToolsManagerCommands.BeginEdgeLoopInsertionTool", 		new IMAGE_PLUGIN_BRUSH("Icons/ModelingEdgeLoopInsert_x40", 		Icon20x20));
@@ -232,8 +232,8 @@ void FModelingToolsEditorModeStyle::Initialize()
 		StyleSet->Set("ModelingToolsManagerCommands.BeginRevolveBoundaryTool.Small",         new IMAGE_PLUGIN_BRUSH("Icons/ModelingRevolveBoundary_x20", Icon20x20));
 		StyleSet->Set("ModelingToolsManagerCommands.BeginMeshBooleanTool",                   new IMAGE_PLUGIN_BRUSH("Icons/ModelingMeshBoolean_x40", Icon20x20));
 		StyleSet->Set("ModelingToolsManagerCommands.BeginMeshBooleanTool.Small",             new IMAGE_PLUGIN_BRUSH("Icons/ModelingMeshBoolean_x20", Icon20x20));
-		//StyleSet->Set("ModelingToolsManagerCommands.BeginMeshBooleanTool",				 new IMAGE_PLUGIN_BRUSH("Icons/ModelingMeshTrim_x40", Icon20x20));
-		//StyleSet->Set("ModelingToolsManagerCommands.BeginMeshBooleanTool.Small",			 new IMAGE_PLUGIN_BRUSH("Icons/ModelingMeshTrim_x20", Icon20x20));
+		StyleSet->Set("ModelingToolsManagerCommands.BeginMeshTrimTool",				 new IMAGE_BRUSH_SVG("Icons/ModelingTrim", Icon20x20));
+		StyleSet->Set("ModelingToolsManagerCommands.BeginMeshTrimTool.Small",			 new IMAGE_BRUSH_SVG("Icons/ModelingTrim", Icon20x20));
 		StyleSet->Set("ModelingToolsManagerCommands.BeginSelfUnionTool",                     new IMAGE_PLUGIN_BRUSH("Icons/ModelingSelfUnion_x40", Icon20x20));
 		StyleSet->Set("ModelingToolsManagerCommands.BeginSelfUnionTool.Small",               new IMAGE_PLUGIN_BRUSH("Icons/ModelingSelfUnion_x20", Icon20x20));
 		StyleSet->Set("ModelingToolsManagerCommands.BeginVoxelSolidifyTool",                 new IMAGE_PLUGIN_BRUSH("Icons/ModelingVoxSolidify_x40", Icon20x20));
@@ -272,6 +272,8 @@ void FModelingToolsEditorModeStyle::Initialize()
 
 		StyleSet->Set("ModelingToolsManagerCommands.BeginGenerateStaticMeshLODAssetTool", new IMAGE_BRUSH_SVG("Icons/AutoLOD", Icon20x20));
 		StyleSet->Set("ModelingToolsManagerCommands.BeginGenerateStaticMeshLODAssetTool.Small", new IMAGE_BRUSH_SVG("Icons/AutoLOD", Icon20x20));
+		StyleSet->Set("ModelingToolsManagerCommands.BeginLODManagerTool", new IMAGE_BRUSH_SVG("Icons/ModelingLODManager", Icon20x20));
+		StyleSet->Set("ModelingToolsManagerCommands.BeginLODManagerTool.Small", new IMAGE_BRUSH_SVG("Icons/ModelingLODManager", Icon20x20));
 
 		StyleSet->Set("ModelingToolsManagerCommands.BeginGroomCardsEditorTool", new IMAGE_BRUSH_SVG("Icons/CardsEditor", Icon20x20));
 		StyleSet->Set("ModelingToolsManagerCommands.BeginGroomCardsEditorTool.Small", new IMAGE_BRUSH_SVG("Icons/CardsEditor", Icon20x20));
