When creating a new static mesh, also write the materials to the UStaticMesh (in addition to the UStaticMeshComponent)

#rnx

diff --git a/Engine/Plugins/Experimental/ModelingToolsEditorMode/Source/ModelingToolsEditorMode/Private/ModelingModeAssetAPI.cpp b/Engine/Plugins/Experimental/ModelingToolsEditorMode/Source/ModelingToolsEditorMode/Private/ModelingModeAssetAPI.cpp
--- a/Engine/Plugins/Experimental/ModelingToolsEditorMode/Source/ModelingToolsEditorMode/Private/ModelingModeAssetAPI.cpp
+++ b/Engine/Plugins/Experimental/ModelingToolsEditorMode/Source/ModelingToolsEditorMode/Private/ModelingModeAssetAPI.cpp
@@ -200,6 +200,11 @@ AActor* FModelingModeAssetAPI::GenerateStaticMeshActor(
 	{
 		Builder.NewMeshComponent->SetMaterial(MatIdx, AssetConfig.Materials[MatIdx]);
 	}
+	// also put the materials on the asset
+	for (int MatIdx = 0, NumMats = AssetConfig.Materials.Num(); MatIdx < NumMats; MatIdx++)
+	{
+		Builder.NewStaticMesh->SetMaterial(MatIdx, AssetConfig.Materials[MatIdx]);
+	}
 
 	// save the new asset (or don't, if that's what the user wants)
 	if (AutoGenMode == EModelingModeAssetGenerationBehavior::AutoGenerateAndAutosave )
