ModelingMode: update Cut Mesh With Mesh tool icon
#rnx

diff --git a/Engine/Plugins/Experimental/ModelingToolsEditorMode/Source/ModelingToolsEditorMode/Private/ModelingToolsEditorModeStyle.cpp b/Engine/Plugins/Experimental/ModelingToolsEditorMode/Source/ModelingToolsEditorMode/Private/ModelingToolsEditorModeStyle.cpp
--- a/Engine/Plugins/Experimental/ModelingToolsEditorMode/Source/ModelingToolsEditorMode/Private/ModelingToolsEditorModeStyle.cpp
+++ b/Engine/Plugins/Experimental/ModelingToolsEditorMode/Source/ModelingToolsEditorMode/Private/ModelingToolsEditorModeStyle.cpp
@@ -234,8 +234,8 @@ void FModelingToolsEditorModeStyle::Initialize()
 		StyleSet->Set("ModelingToolsManagerCommands.BeginMeshBooleanTool.Small",             new IMAGE_PLUGIN_BRUSH("Icons/ModelingMeshBoolean_x20", Icon20x20));
 		StyleSet->Set("ModelingToolsManagerCommands.BeginMeshTrimTool",				 new IMAGE_BRUSH_SVG("Icons/ModelingTrim", Icon20x20));
 		StyleSet->Set("ModelingToolsManagerCommands.BeginMeshTrimTool.Small",			 new IMAGE_BRUSH_SVG("Icons/ModelingTrim", Icon20x20));
-		StyleSet->Set("ModelingToolsManagerCommands.BeginCutMeshWithMeshTool",               new IMAGE_BRUSH_SVG("Icons/ModelingTrim", Icon20x20));
-		StyleSet->Set("ModelingToolsManagerCommands.BeginCutMeshWithMeshTool.Small",         new IMAGE_BRUSH_SVG("Icons/ModelingTrim", Icon20x20));
+		StyleSet->Set("ModelingToolsManagerCommands.BeginCutMeshWithMeshTool",               new IMAGE_BRUSH_SVG("Icons/ModelingMeshCut", Icon20x20));
+		StyleSet->Set("ModelingToolsManagerCommands.BeginCutMeshWithMeshTool.Small",         new IMAGE_BRUSH_SVG("Icons/ModelingMeshCut", Icon20x20));
 		StyleSet->Set("ModelingToolsManagerCommands.BeginSelfUnionTool",                     new IMAGE_PLUGIN_BRUSH("Icons/ModelingSelfUnion_x40", Icon20x20));
 		StyleSet->Set("ModelingToolsManagerCommands.BeginSelfUnionTool.Small",               new IMAGE_PLUGIN_BRUSH("Icons/ModelingSelfUnion_x20", Icon20x20));
 		StyleSet->Set("ModelingToolsManagerCommands.BeginVoxelSolidifyTool",                 new IMAGE_PLUGIN_BRUSH("Icons/ModelingVoxSolidify_x40", Icon20x20));
