ModelingMode: add support for extending ModelingMode with additional Tools based on IModularFeature framework. Tool Extensions appear in tool icons panel in extra sections below the standard sections, currently adding to default Sections is not supported.
#rnx

diff --git a/Engine/Plugins/Experimental/ModelingToolsEditorMode/Source/ModelingToolsEditorMode/Private/ModelingToolsEditorMode.cpp b/Engine/Plugins/Experimental/ModelingToolsEditorMode/Source/ModelingToolsEditorMode/Private/ModelingToolsEditorMode.cpp
--- a/Engine/Plugins/Experimental/ModelingToolsEditorMode/Source/ModelingToolsEditorMode/Private/ModelingToolsEditorMode.cpp
+++ b/Engine/Plugins/Experimental/ModelingToolsEditorMode/Source/ModelingToolsEditorMode/Private/ModelingToolsEditorMode.cpp
@@ -14,6 +14,9 @@
 #include "EditorViewportClient.h"
 #include "EngineAnalytics.h"
 
+#include "Features/IModularFeatures.h"
+#include "ModelingModeToolExtensions.h"
+
 //#include "SingleClickTool.h"
 //#include "MeshSurfacePointTool.h"
 //#include "MeshVertexDragTool.h"
@@ -604,6 +607,33 @@ void UModelingToolsEditorMode::Enter()
 	RegisterPolyModelActionTool(EEditMeshPolygonsToolActions::CutFaces, ToolManagerCommands.BeginPolyModelTool_CutFaces, TEXT("PolyEdit_CutFaces"));
 
 
+	// register extensions
+	TArray<IModelingModeToolExtension*> Extensions = IModularFeatures::Get().GetModularFeatureImplementations<IModelingModeToolExtension>(
+		IModelingModeToolExtension::GetModularFeatureName());
+	if (Extensions.Num() > 0)
+	{
+		FExtensionToolQueryInfo ExtensionQueryInfo;
+		ExtensionQueryInfo.ToolsContext = ToolsContext.Get();
+		ExtensionQueryInfo.AssetAPI = ModelingModeAssetGenerationAPI.Get();
+
+		UE_LOG(LogTemp, Log, TEXT("ModelingMode: Found %d Tool Extension Modules"), Extensions.Num());
+		for (int32 k = 0; k < Extensions.Num(); ++k)
+		{
+			// TODO: extension name
+			FText ExtensionName = Extensions[k]->GetExtensionName();
+			FString ExtensionPrefix = FString::Printf(TEXT("[%d][%s]"), k, *ExtensionName.ToString());
+
+			TArray<FExtensionToolDescription> ToolSet;
+			Extensions[k]->GetExtensionTools(ExtensionQueryInfo, ToolSet);
+			for (const FExtensionToolDescription& ToolInfo : ToolSet)
+			{
+				UE_LOG(LogTemp, Log, TEXT("%s - Registering Tool [%s]"), *ExtensionPrefix, *ToolInfo.ToolName.ToString());
+
+				RegisterTool(ToolInfo.ToolCommand, ToolInfo.ToolName.ToString(), ToolInfo.ToolBuilder);
+			}
+		}
+	}
+
 
 	ToolsContext->ToolManager->SelectActiveToolType(EToolSide::Left, TEXT("DynaSculptTool"));
 
diff --git a/Engine/Plugins/Experimental/ModelingToolsEditorMode/Source/ModelingToolsEditorMode/Private/ModelingToolsEditorModeToolkit.cpp b/Engine/Plugins/Experimental/ModelingToolsEditorMode/Source/ModelingToolsEditorMode/Private/ModelingToolsEditorModeToolkit.cpp
--- a/Engine/Plugins/Experimental/ModelingToolsEditorMode/Source/ModelingToolsEditorMode/Private/ModelingToolsEditorModeToolkit.cpp
+++ b/Engine/Plugins/Experimental/ModelingToolsEditorMode/Source/ModelingToolsEditorMode/Private/ModelingToolsEditorModeToolkit.cpp
@@ -21,6 +21,10 @@
 #include "Widgets/Images/SImage.h"
 #include "Widgets/Layout/SExpandableArea.h"
 
+// for Tool Extensions
+#include "Features/IModularFeatures.h"
+#include "ModelingModeToolExtensions.h"
+
 // for LOD setting
 #include "EditorInteractiveToolsFrameworkModule.h"
 #include "Tools/EditorComponentSourceFactory.h"
@@ -53,6 +57,12 @@ FModelingToolsEditorModeToolkit::~FModelingToolsEditorModeToolkit()
 	GetScriptableEditorMode()->GetInteractiveToolsContext()->OnToolWarningMessage.RemoveAll(this);
 }
 
+
+void FModelingToolsEditorModeToolkit::CustomizeModeDetailsViewArgs(FDetailsViewArgs& ArgsInOut)
+{
+	//ArgsInOut.ColumnWidth = 0.3f;
+}
+
 void FModelingToolsEditorModeToolkit::Init(const TSharedPtr<IToolkitHost>& InitToolkitHost, TWeakObjectPtr<UEdMode> InOwningMode)
 {
 	// Have to create the ToolkitWidget here because FModeToolkit::Init() is going to ask for it and add
@@ -429,16 +439,46 @@ void FModelingToolsEditorModeToolkit::GetToolPaletteNames(TArray<FName>& Palette
 {
 	PaletteNames = PaletteNames_Standard;
 
+	TArray<FName> ExistingNames;
+	for ( FName Name : PaletteNames )
+	{
+		ExistingNames.Add(Name);
+	}
+
+
 	bool bEnablePrototypes = (CVarEnablePrototypeModelingTools.GetValueOnGameThread() > 0);
 	if (bEnablePrototypes)
 	{
 		PaletteNames.Add(PrototypesTabName);
+		ExistingNames.Add(PrototypesTabName);
 	}
 
 	bool bEnablePolyModel = (CVarEnablePolyModeling.GetValueOnGameThread() > 0);
 	if (bEnablePolyModel)
 	{
 		PaletteNames.Add(PolyEditTabName);
+		ExistingNames.Add(PolyEditTabName);
+	}
+
+	if (IModularFeatures::Get().IsModularFeatureAvailable(IModelingModeToolExtension::GetModularFeatureName()))
+	{
+		TArray<IModelingModeToolExtension*> Extensions = IModularFeatures::Get().GetModularFeatureImplementations<IModelingModeToolExtension>(
+			IModelingModeToolExtension::GetModularFeatureName());
+		for (int32 k = 0; k < Extensions.Num(); ++k)
+		{
+			FText ExtensionName = Extensions[k]->GetExtensionName();
+			FText SectionName = Extensions[k]->GetToolSectionName();
+			FName SectionIndex(SectionName.ToString());
+			if (ExistingNames.Contains(SectionIndex))
+			{
+				UE_LOG(LogTemp, Warning, TEXT("Modeling Mode Extension [%s] uses existing Section Name [%s] - buttons may not be visible"), *ExtensionName.ToString(), *SectionName.ToString());
+			}
+			else
+			{
+				PaletteNames.Add(SectionIndex);
+				ExistingNames.Add(SectionIndex);
+			}
+		}
 	}
 }
 
@@ -611,6 +651,27 @@ void FModelingToolsEditorModeToolkit::BuildToolPalette_Experimental(FName Palett
 		ToolbarBuilder.AddToolBarButton(Commands.BeginSubdividePolyTool);
 		ToolbarBuilder.AddToolBarButton(Commands.BeginPolyEditTool);
 	}
+	else
+	{
+		TArray<IModelingModeToolExtension*> Extensions = IModularFeatures::Get().GetModularFeatureImplementations<IModelingModeToolExtension>(
+			IModelingModeToolExtension::GetModularFeatureName());
+		for (int32 k = 0; k < Extensions.Num(); ++k)
+		{
+			FText SectionName = Extensions[k]->GetToolSectionName();
+			FName SectionIndex(SectionName.ToString());
+			if (PaletteIndex == SectionIndex)
+			{
+				FExtensionToolQueryInfo ExtensionQueryInfo;
+				ExtensionQueryInfo.bIsInfoQueryOnly = true;
+				TArray<FExtensionToolDescription> ToolSet;
+				Extensions[k]->GetExtensionTools(ExtensionQueryInfo, ToolSet);
+				for (const FExtensionToolDescription& ToolInfo : ToolSet)
+				{
+					ToolbarBuilder.AddToolBarButton(ToolInfo.ToolCommand);
+				}
+			}
+		}
+	}
 }
 
 void FModelingToolsEditorModeToolkit::OnToolPaletteChanged(FName PaletteName) 
diff --git a/Engine/Plugins/Experimental/ModelingToolsEditorMode/Source/ModelingToolsEditorMode/Public/ModelingModeToolExtensions.h b/Engine/Plugins/Experimental/ModelingToolsEditorMode/Source/ModelingToolsEditorMode/Public/ModelingModeToolExtensions.h
new file mode 100644
--- /dev/null
+++ b/Engine/Plugins/Experimental/ModelingToolsEditorMode/Source/ModelingToolsEditorMode/Public/ModelingModeToolExtensions.h
@@ -0,0 +1,86 @@
+// Copyright Epic Games, Inc. All Rights Reserved.
+
+#pragma once
+
+#include "CoreMinimal.h"
+#include "Features/IModularFeature.h"
+
+class UInteractiveToolsContext;
+class FModelingModeAssetAPI;
+class FUICommandInfo;
+class UInteractiveToolBuilder;
+
+
+/**
+ * This struct is passed to IModelingModeToolExtension implementations to allow 
+ * them to forward various Tools Context information to their ToolBuilders/etc
+ * 
+ * Note that if bIsInfoQueryOnly is true, then the Extension does not need to 
+ * construct Tool Builders, the query is for information purposes only
+ */
+struct MODELINGTOOLSEDITORMODE_API FExtensionToolQueryInfo
+{
+	bool bIsInfoQueryOnly = false;
+
+	UInteractiveToolsContext* ToolsContext = nullptr;
+	FModelingModeAssetAPI* AssetAPI = nullptr;
+};
+
+/**
+ * IModelingModeToolExtension implementations return the list of Tools they provide
+ * via instances of FExtensionToolDescription.
+ */
+struct MODELINGTOOLSEDITORMODE_API FExtensionToolDescription
+{
+	/** Long name of the Tool, used in various places in the UI */
+	FText ToolName;
+	/** Command that is added to the Tool button set. This defines the button label. */
+	TSharedPtr<FUICommandInfo> ToolCommand;
+	/** Builder for the Tool. This can be null if FExtensionToolQueryInfo.bIsInfoQueryOnly is true. */
+	UInteractiveToolBuilder* ToolBuilder = nullptr;
+};
+
+/**
+ * IModelingModeToolExtension uses the IModularFeature API to allow a Plugin to provide
+ * a set of InteractiveTool's to be exposed in Modeling Mode. The Tools will be 
+ * included in a section of the Modeling Mode tool list, based on GetToolSectionName().
+ * 
+ */
+class MODELINGTOOLSEDITORMODE_API IModelingModeToolExtension : public IModularFeature
+{
+public:
+	virtual ~IModelingModeToolExtension() {}
+
+	/**
+	 * @return A text string that identifies this Extension.
+	 */
+	virtual FText GetExtensionName() = 0;
+
+	/**
+	 * @return A text string that defines the name of the Toolbar Section this Extension's tools will be included in.
+	 * @warning If the same Section is used in multiple Extensions, some buttons may not be shown
+	 */
+	virtual FText GetToolSectionName() = 0;
+
+	/**
+	 * Query the Extension for a list of Tools to expose in Modeling Mode.
+	 * Note that this function *will* be called multiple times by Modeling Mode, as the
+	 * information about the set of Tools is needed in multiple places. The QueryInfo.bIsInfoQueryOnly
+	 * flag indicates whether the caller requires ToolBuilder instances. 
+	 * 
+	 * If creating multiple copies of the ToolBuilder for a particular Tool would be problematic, it 
+	 * is the responsiblity of the IModelingModeToolExtension implementation to cache these internally, 
+	 * otherwise they will be garbage collected when the caller releases them.
+	 */
+	virtual void GetExtensionTools(const FExtensionToolQueryInfo& QueryInfo, TArray<FExtensionToolDescription>& ToolsOut) = 0;
+
+
+
+	static FName GetModularFeatureName()
+	{
+		static FName FeatureName = FName(TEXT("ModelingModeToolExtension"));
+		return FeatureName;
+	}
+};
+
+
diff --git a/Engine/Plugins/Experimental/ModelingToolsEditorMode/Source/ModelingToolsEditorMode/Public/ModelingToolsEditorModeStyle.h b/Engine/Plugins/Experimental/ModelingToolsEditorMode/Source/ModelingToolsEditorMode/Public/ModelingToolsEditorModeStyle.h
--- a/Engine/Plugins/Experimental/ModelingToolsEditorMode/Source/ModelingToolsEditorMode/Public/ModelingToolsEditorModeStyle.h
+++ b/Engine/Plugins/Experimental/ModelingToolsEditorMode/Source/ModelingToolsEditorMode/Public/ModelingToolsEditorModeStyle.h
@@ -6,7 +6,7 @@
 #include "Styling/ISlateStyle.h"
 #include "Styling/SlateStyle.h"
 
-class FModelingToolsEditorModeStyle
+class MODELINGTOOLSEDITORMODE_API FModelingToolsEditorModeStyle
 {
 public:
 	static void Initialize();
diff --git a/Engine/Plugins/Experimental/ModelingToolsEditorMode/Source/ModelingToolsEditorMode/Public/ModelingToolsEditorModeToolkit.h b/Engine/Plugins/Experimental/ModelingToolsEditorMode/Source/ModelingToolsEditorMode/Public/ModelingToolsEditorModeToolkit.h
--- a/Engine/Plugins/Experimental/ModelingToolsEditorMode/Source/ModelingToolsEditorMode/Public/ModelingToolsEditorModeToolkit.h
+++ b/Engine/Plugins/Experimental/ModelingToolsEditorMode/Source/ModelingToolsEditorMode/Public/ModelingToolsEditorModeToolkit.h
@@ -58,6 +58,8 @@ public:
 	virtual void OnToolStarted(UInteractiveToolManager* Manager, UInteractiveTool* Tool) override;
 	virtual void OnToolEnded(UInteractiveToolManager* Manager, UInteractiveTool* Tool) override;
 
+	virtual void CustomizeModeDetailsViewArgs(FDetailsViewArgs& ArgsInOut) override;
+
 	void OnActiveViewportChanged(TSharedPtr<IAssetViewport>, TSharedPtr<IAssetViewport> );
 
 
