ModelingMode: Add support for selection-focus and point-focus at Tool level.
- New UInteractiveToolCameraFocusAPI UInterface defines functions a Tool can implement to publish its ability to provide a focus-box and focus-point
- ModelingMode uses this API to customize behavior. Box-focus (f hotkey) will now always focus on Tool-provided focus-box if available. Point-focus (c hotkey) will focus on nearer of world-hit and Tool focus-point.
- Default implementations for PrimitiveComponentToolTarget in SingleSelectionTool and MultiSelectionTool
- Custom focus-box implementations for PolyEdit-derived and TriSelect-derived tools, to focus on active selection if available
- Custom focus-point implementations for MeshSurfacePointTool

#rnx

diff --git a/Engine/Plugins/Experimental/MeshModelingToolset/Source/MeshModelingTools/Private/EditMeshPolygonsTool.cpp b/Engine/Plugins/Experimental/MeshModelingToolset/Source/MeshModelingTools/Private/EditMeshPolygonsTool.cpp
--- a/Engine/Plugins/Experimental/MeshModelingToolset/Source/MeshModelingTools/Private/EditMeshPolygonsTool.cpp
+++ b/Engine/Plugins/Experimental/MeshModelingToolset/Source/MeshModelingTools/Private/EditMeshPolygonsTool.cpp
@@ -653,6 +653,34 @@ void UEditMeshPolygonsTool::UpdateMultiTransformerFrame(const FFrame3d* UseFrame
 }
 
 
+FBox UEditMeshPolygonsTool::GetWorldSpaceFocusBox()
+{
+	if (SelectionMechanic && SelectionMechanic->HasSelection())
+	{
+		FAxisAlignedBox3d Bounds = SelectionMechanic->GetSelectionBounds(true);
+		return (FBox)Bounds;
+	}
+	return UMeshSurfacePointTool::GetWorldSpaceFocusBox();
+}
+
+bool UEditMeshPolygonsTool::GetWorldSpaceFocusPoint(const FRay& WorldRay, FVector& PointOut)
+{
+	FRay3d LocalRay(WorldTransform.InverseTransformPosition((FVector3d)WorldRay.Origin),
+		WorldTransform.InverseTransformNormal((FVector3d)WorldRay.Direction));
+	UE::Geometry::Normalize(LocalRay.Direction);
+
+	int32 HitTID = GetSpatial().FindNearestHitTriangle(LocalRay);
+	if (HitTID != IndexConstants::InvalidID)
+	{
+		FIntrRay3Triangle3d TriHit = TMeshQueries<FDynamicMesh3>::TriangleIntersection(*GetSpatial().GetMesh(), HitTID, LocalRay);
+		FVector3d LocalPos = LocalRay.PointAt(TriHit.RayParameter);
+		PointOut = (FVector)WorldTransform.TransformPosition(LocalPos);
+		return true;
+	}
+	return false;
+}
+
+
 void UEditMeshPolygonsTool::OnSelectionModifiedEvent()
 {
 	FVector3d LocalLastHitPosition, LocalLastHitNormal;
diff --git a/Engine/Plugins/Experimental/MeshModelingToolset/Source/MeshModelingTools/Private/MeshSelectionTool.cpp b/Engine/Plugins/Experimental/MeshModelingToolset/Source/MeshModelingTools/Private/MeshSelectionTool.cpp
--- a/Engine/Plugins/Experimental/MeshModelingToolset/Source/MeshModelingTools/Private/MeshSelectionTool.cpp
+++ b/Engine/Plugins/Experimental/MeshModelingToolset/Source/MeshModelingTools/Private/MeshSelectionTool.cpp
@@ -663,6 +663,34 @@ bool UMeshSelectionTool::OnUpdateHover(const FInputDeviceRay& DevicePos)
 }
 
 
+FBox UMeshSelectionTool::GetWorldSpaceFocusBox()
+{
+	TArray<int32> SelectedFaces = Selection->GetElements(EMeshSelectionElementType::Face);
+	if (SelectedFaces.Num() > 0)
+	{
+		FAxisAlignedBox3d Bounds = FAxisAlignedBox3d::Empty();
+
+		const FDynamicMesh3* Mesh = PreviewMesh->GetMesh();
+		FTransform3d Transform(PreviewMesh->GetTransform());
+		for (int32 tid : SelectedFaces)
+		{
+			FIndex3i Tri = Mesh->GetTriangle(tid);
+			for (int32 j = 0; j < 3; ++j)
+			{
+				Bounds.Contain(Transform.TransformPosition(Mesh->GetVertex(Tri[j])));
+			}
+		}
+
+		if (Bounds.MaxDim() > FMathf::ZeroTolerance)
+		{
+			return (FBox)Bounds;
+		}
+	}
+
+	return PreviewMesh->GetActor()->GetComponentsBoundingBox();
+}
+
+
 
 
 void UMeshSelectionTool::OnRegionHighlightUpdated()
diff --git a/Engine/Plugins/Experimental/MeshModelingToolset/Source/MeshModelingTools/Public/EditMeshPolygonsTool.h b/Engine/Plugins/Experimental/MeshModelingToolset/Source/MeshModelingTools/Public/EditMeshPolygonsTool.h
--- a/Engine/Plugins/Experimental/MeshModelingToolset/Source/MeshModelingTools/Public/EditMeshPolygonsTool.h
+++ b/Engine/Plugins/Experimental/MeshModelingToolset/Source/MeshModelingTools/Public/EditMeshPolygonsTool.h
@@ -528,6 +528,10 @@ public:
 	virtual FInputRayHit IsHitByClick(const FInputDeviceRay& ClickPos) override;
 	virtual void OnClicked(const FInputDeviceRay& ClickPos) override;
 
+	// IInteractiveToolCameraFocusAPI implementation
+	virtual FBox GetWorldSpaceFocusBox() override;
+	virtual bool GetWorldSpaceFocusPoint(const FRay& WorldRay, FVector& PointOut) override;
+
 public:
 
 	virtual void RequestAction(EEditMeshPolygonsToolActions ActionType);
diff --git a/Engine/Plugins/Experimental/MeshModelingToolset/Source/MeshModelingTools/Public/MeshSelectionTool.h b/Engine/Plugins/Experimental/MeshModelingToolset/Source/MeshModelingTools/Public/MeshSelectionTool.h
--- a/Engine/Plugins/Experimental/MeshModelingToolset/Source/MeshModelingTools/Public/MeshSelectionTool.h
+++ b/Engine/Plugins/Experimental/MeshModelingToolset/Source/MeshModelingTools/Public/MeshSelectionTool.h
@@ -291,6 +291,9 @@ public:
 	virtual void OnEndDrag(const FRay& Ray) override;
 	virtual bool OnUpdateHover(const FInputDeviceRay& DevicePos) override;
 
+	// IInteractiveToolCameraFocusAPI implementation
+	virtual FBox GetWorldSpaceFocusBox() override;
+
 public:
 
 	virtual void RequestAction(EMeshSelectionToolActions ActionType);
diff --git a/Engine/Plugins/Experimental/ModelingToolsEditorMode/Source/ModelingToolsEditorMode/Private/ModelingToolsEditorMode.cpp b/Engine/Plugins/Experimental/ModelingToolsEditorMode/Source/ModelingToolsEditorMode/Private/ModelingToolsEditorMode.cpp
--- a/Engine/Plugins/Experimental/ModelingToolsEditorMode/Source/ModelingToolsEditorMode/Private/ModelingToolsEditorMode.cpp
+++ b/Engine/Plugins/Experimental/ModelingToolsEditorMode/Source/ModelingToolsEditorMode/Private/ModelingToolsEditorMode.cpp
@@ -779,16 +779,70 @@ void UModelingToolsEditorMode::FocusCameraAtCursorHotkey()
 {
 	FRay Ray = ToolsContext->GetLastWorldRay();
 
-	FHitResult HitResult;
-	bool bHitWorld = ToolSceneQueriesUtil::FindNearestVisibleObjectHit(GetWorld(), HitResult, Ray.Origin, Ray.PointAt(HALF_WORLD_MAX));
-	if (bHitWorld)
+	double NearestHitDist = (double)HALF_WORLD_MAX;
+	FVector HitPoint = FVector::ZeroVector;
+
+	// cast ray against visible objects
+	FHitResult WorldHitResult;
+	if (ToolSceneQueriesUtil::FindNearestVisibleObjectHit(GetWorld(), WorldHitResult, Ray.Origin, Ray.PointAt(HALF_WORLD_MAX)))
+	{
+		HitPoint = WorldHitResult.ImpactPoint;
+		NearestHitDist = (double)Ray.GetParameter(HitPoint);
+	}
+
+	// cast ray against tool
+	if (GetToolManager()->HasAnyActiveTool())
+	{
+		UInteractiveTool* Tool = GetToolManager()->GetActiveTool(EToolSide::Mouse);
+		IInteractiveToolCameraFocusAPI* FocusAPI = Cast<IInteractiveToolCameraFocusAPI>(Tool);
+		if (FocusAPI && FocusAPI->SupportsWorldSpaceFocusPoint())
+		{
+			FVector ToolHitPoint;
+			if (FocusAPI->GetWorldSpaceFocusPoint(Ray, ToolHitPoint))
+			{
+				double HitDepth = (double)Ray.GetParameter(ToolHitPoint);
+				if (HitDepth < NearestHitDist)
+				{
+					NearestHitDist = HitDepth;
+					HitPoint = ToolHitPoint;
+				}
+			}
+		}
+	}
+
+
+	if (NearestHitDist < (double)HALF_WORLD_MAX && GCurrentLevelEditingViewportClient)
 	{
-		FVector HitPoint = HitResult.ImpactPoint;
-		if (GCurrentLevelEditingViewportClient)
+		GCurrentLevelEditingViewportClient->CenterViewportAtPoint(HitPoint, false);
+	}
+}
+
+
+bool UModelingToolsEditorMode::ComputeBoundingBoxForViewportFocus(AActor* Actor, UPrimitiveComponent* PrimitiveComponent, FBox& InOutBox) const
+{
+	if (GetToolManager()->HasAnyActiveTool())
+	{
+		UInteractiveTool* Tool = GetToolManager()->GetActiveTool(EToolSide::Mouse);
+		IInteractiveToolCameraFocusAPI* FocusAPI = Cast<IInteractiveToolCameraFocusAPI>(Tool);
+		if (FocusAPI && FocusAPI->SupportsWorldSpaceFocusBox() )
 		{
-			GCurrentLevelEditingViewportClient->CenterViewportAtPoint(HitPoint, false);
+			InOutBox = FocusAPI->GetWorldSpaceFocusBox();
+			if (InOutBox.IsValid)
+			{
+				float MaxDimension = InOutBox.GetExtent().GetMax();
+				if (MaxDimension > SMALL_NUMBER)
+				{
+					InOutBox.ExpandBy(MaxDimension * 0.2f);
+				}
+				else
+				{
+					InOutBox.ExpandBy(25);
+				}
+				return true;
+			}
 		}
 	}
+	return false;
 }
 
 
diff --git a/Engine/Plugins/Experimental/ModelingToolsEditorMode/Source/ModelingToolsEditorMode/Public/ModelingToolsEditorMode.h b/Engine/Plugins/Experimental/ModelingToolsEditorMode/Source/ModelingToolsEditorMode/Public/ModelingToolsEditorMode.h
--- a/Engine/Plugins/Experimental/ModelingToolsEditorMode/Source/ModelingToolsEditorMode/Public/ModelingToolsEditorMode.h
+++ b/Engine/Plugins/Experimental/ModelingToolsEditorMode/Source/ModelingToolsEditorMode/Public/ModelingToolsEditorMode.h
@@ -40,6 +40,8 @@ public:
 
 	virtual bool CanAutoSave() const override;
 
+	virtual bool ComputeBoundingBoxForViewportFocus(AActor* Actor, UPrimitiveComponent* PrimitiveComponent, FBox& InOutBox) const override;
+
 	virtual bool GetPivotForOrbit(FVector& OutPivot) const override;
 
 	/*
diff --git a/Engine/Source/Runtime/Experimental/InteractiveToolsFramework/Private/BaseTools/MeshSurfacePointTool.cpp b/Engine/Source/Runtime/Experimental/InteractiveToolsFramework/Private/BaseTools/MeshSurfacePointTool.cpp
--- a/Engine/Source/Runtime/Experimental/InteractiveToolsFramework/Private/BaseTools/MeshSurfacePointTool.cpp
+++ b/Engine/Source/Runtime/Experimental/InteractiveToolsFramework/Private/BaseTools/MeshSurfacePointTool.cpp
@@ -96,6 +96,19 @@ bool UMeshSurfacePointTool::HitTest(const FRay& Ray, FHitResult& OutHit)
 }
 
 
+bool UMeshSurfacePointTool::GetWorldSpaceFocusPoint(const FRay& WorldRay, FVector& PointOut)
+{
+	FHitResult HitResult;
+	if (HitTest(WorldRay, HitResult))
+	{
+		PointOut = HitResult.ImpactPoint;
+		return true;
+	}
+	return false;
+}
+
+
+
 void UMeshSurfacePointTool::OnBeginDrag(const FRay& Ray)
 {
 
diff --git a/Engine/Source/Runtime/Experimental/InteractiveToolsFramework/Private/BaseTools/SingleSelectionTool.cpp b/Engine/Source/Runtime/Experimental/InteractiveToolsFramework/Private/BaseTools/SingleSelectionTool.cpp
new file mode 100644
--- /dev/null
+++ b/Engine/Source/Runtime/Experimental/InteractiveToolsFramework/Private/BaseTools/SingleSelectionTool.cpp
@@ -0,0 +1,60 @@
+// Copyright Epic Games, Inc. All Rights Reserved.
+
+#include "SingleSelectionTool.h"
+#include "InteractiveToolManager.h"
+#include "ToolTargets/PrimitiveComponentToolTarget.h"
+#include "Components/PrimitiveComponent.h"
+
+#define LOCTEXT_NAMESPACE "USingleSelectionTool"
+
+bool USingleSelectionTool::SupportsWorldSpaceFocusBox()
+{
+	return Target && Cast<UPrimitiveComponentToolTarget>(Target) != nullptr;
+}
+
+
+FBox USingleSelectionTool::GetWorldSpaceFocusBox()
+{
+	if (Target)
+	{
+		UPrimitiveComponentToolTarget* PrimTarget = Cast<UPrimitiveComponentToolTarget>(Target);
+		if (PrimTarget)
+		{
+			UPrimitiveComponent* Component = PrimTarget->GetOwnerComponent();
+			if (Component)
+			{
+				return Component->Bounds.GetBox();
+			}
+		}
+	}
+	return FBox();
+}
+
+
+bool USingleSelectionTool::SupportsWorldSpaceFocusPoint()
+{
+	return Target && Cast<UPrimitiveComponentToolTarget>(Target) != nullptr;
+
+}
+bool USingleSelectionTool::GetWorldSpaceFocusPoint(const FRay& WorldRay, FVector& PointOut)
+{
+	if (Target)
+	{
+		UPrimitiveComponentToolTarget* PrimTarget = Cast<UPrimitiveComponentToolTarget>(Target);
+		if (PrimTarget)
+		{
+			FHitResult HitResult;
+			if (PrimTarget->HitTestComponent(WorldRay, HitResult))
+			{
+				PointOut = HitResult.ImpactPoint;
+				return true;
+			}
+		}
+	}
+	return false;
+}
+
+
+
+
+#undef LOCTEXT_NAMESPACE
\ No newline at end of file
diff --git a/Engine/Source/Runtime/Experimental/InteractiveToolsFramework/Private/MultiSelectionTool.cpp b/Engine/Source/Runtime/Experimental/InteractiveToolsFramework/Private/MultiSelectionTool.cpp
--- a/Engine/Source/Runtime/Experimental/InteractiveToolsFramework/Private/MultiSelectionTool.cpp
+++ b/Engine/Source/Runtime/Experimental/InteractiveToolsFramework/Private/MultiSelectionTool.cpp
@@ -1,6 +1,8 @@
 // Copyright Epic Games, Inc. All Rights Reserved.
 #include "MultiSelectionTool.h"
 #include "TargetInterfaces/AssetBackedTarget.h"
+#include "ToolTargets/PrimitiveComponentToolTarget.h"
+#include "Components/PrimitiveComponent.h"
 
 bool UMultiSelectionTool::GetMapToSharedSourceData(TArray<int32>& MapToFirstOccurrences)
 {
@@ -37,3 +39,77 @@ bool UMultiSelectionTool::GetMapToSharedSourceData(TArray<int32>& MapToFirstOccu
 	return bSharesSources;
 }
 
+
+
+bool UMultiSelectionTool::SupportsWorldSpaceFocusBox()
+{
+	int32 PrimitiveCount = 0;
+	for (const TObjectPtr<UToolTarget>& Target : Targets)
+	{
+		if (Cast<UPrimitiveComponentToolTarget>(Target) != nullptr)
+		{
+			PrimitiveCount++;
+		}
+	}
+	return PrimitiveCount > 0;
+}
+
+
+FBox UMultiSelectionTool::GetWorldSpaceFocusBox()
+{
+	FBox AccumBox(EForceInit::ForceInit);
+	for (const TObjectPtr<UToolTarget>& Target : Targets)
+	{
+		UPrimitiveComponentToolTarget* PrimTarget = Cast<UPrimitiveComponentToolTarget>(Target);
+		if (PrimTarget)
+		{
+			UPrimitiveComponent* Component = PrimTarget->GetOwnerComponent();
+			if (Component)
+			{
+				FBox ComponentBounds = Component->Bounds.GetBox();
+				AccumBox += ComponentBounds;
+			}
+		}
+	}
+	return AccumBox;
+}
+
+
+bool UMultiSelectionTool::SupportsWorldSpaceFocusPoint()
+{
+	int32 PrimitiveCount = 0;
+	for (const TObjectPtr<UToolTarget>& Target : Targets)
+	{
+		if (Cast<UPrimitiveComponentToolTarget>(Target) != nullptr)
+		{
+			PrimitiveCount++;
+		}
+	}
+	return PrimitiveCount > 0;
+}
+
+bool UMultiSelectionTool::GetWorldSpaceFocusPoint(const FRay& WorldRay, FVector& PointOut)
+{
+	double NearestRayParam = (double)HALF_WORLD_MAX;
+	PointOut = FVector::ZeroVector;
+
+	for (const TObjectPtr<UToolTarget>& Target : Targets)
+	{
+		UPrimitiveComponentToolTarget* PrimTarget = Cast<UPrimitiveComponentToolTarget>(Target);
+		if (PrimTarget)
+		{
+			FHitResult HitResult;
+			if (PrimTarget->HitTestComponent(WorldRay, HitResult))
+			{
+				double HitRayParam = (double)WorldRay.GetParameter(HitResult.ImpactPoint);
+				if (HitRayParam < NearestRayParam)
+				{
+					NearestRayParam = HitRayParam;
+					PointOut = HitResult.ImpactPoint;
+				}
+			}
+		}
+	}
+
+	return (NearestRayParam < (double)HALF_WORLD_MAX);
+}
diff --git a/Engine/Source/Runtime/Experimental/InteractiveToolsFramework/Public/BaseTools/MeshSurfacePointTool.h b/Engine/Source/Runtime/Experimental/InteractiveToolsFramework/Public/BaseTools/MeshSurfacePointTool.h
--- a/Engine/Source/Runtime/Experimental/InteractiveToolsFramework/Public/BaseTools/MeshSurfacePointTool.h
+++ b/Engine/Source/Runtime/Experimental/InteractiveToolsFramework/Public/BaseTools/MeshSurfacePointTool.h
@@ -115,6 +115,8 @@ public:
 	virtual void OnEndHover() override {}
 
 
+	// IInteractiveToolCameraFocusAPI implementation
+	virtual bool GetWorldSpaceFocusPoint(const FRay& WorldRay, FVector& PointOut) override;
 
 
 	/** Called by registered InputBehaviors to set the state of the "shift" button (or device equivalent) */
diff --git a/Engine/Source/Runtime/Experimental/InteractiveToolsFramework/Public/InteractiveToolQueryInterfaces.h b/Engine/Source/Runtime/Experimental/InteractiveToolsFramework/Public/InteractiveToolQueryInterfaces.h
new file mode 100644
--- /dev/null
+++ b/Engine/Source/Runtime/Experimental/InteractiveToolsFramework/Public/InteractiveToolQueryInterfaces.h
@@ -0,0 +1,67 @@
+// Copyright Epic Games, Inc. All Rights Reserved.
+
+#pragma once
+
+#include "CoreMinimal.h"
+#include "UObject/Interface.h"
+#include "InteractiveToolQueryInterfaces.generated.h"
+
+
+//
+// Below are various interfaces that a UInteractiveTool can implement to allow
+// higher-level code (eg like an EdMode) to query into the Tool.
+//
+
+
+
+// UInterface for IInteractiveToolCameraFocusAPI
+UINTERFACE(MinimalAPI)
+class UInteractiveToolCameraFocusAPI : public UInterface
+{
+	GENERATED_BODY()
+};
+
+/**
+ * IInteractiveToolCameraFocusAPI provides two functions that can be
+ * used to extract "Focus" / "Region of Interest" information about an
+ * active Tool:
+ * 
+ * GetWorldSpaceFocusBox() - provides a bounding box for an "active region" if one is known.
+ *   An example of using the FocusBox would be to center/zoom the camera in a 3D viewport
+ *   onto this box when the user hits a hotkey (eg 'f' in the Editor).
+ *   Should default to the entire active object, if no subregion is available.
+ * 
+ * GetWorldSpaceFocusPoint() - provides a "Focus Point" at the cursor ray if one is known.
+ *   This can be used to (eg) center the camera at the focus point.
+ * 
+ * The above functions should not be called unless the corresponding SupportsX() function returns true.
+ */
+class IInteractiveToolCameraFocusAPI
+{
+	GENERATED_BODY()
+public:
+
+	/**
+	 * @return true if the implementation can provide a Focus Box
+	 */
+	virtual bool SupportsWorldSpaceFocusBox() { return false; }
+	
+	/**
+	 * @return the current Focus Box
+	 */
+	virtual FBox GetWorldSpaceFocusBox() { return FBox(); }
+
+	/**
+	 * @return true if the implementation can provide a Focus Point
+	 */
+	virtual bool SupportsWorldSpaceFocusPoint() { return false; }
+
+	/**
+	 * @param WorldRay 3D Ray that should be used to find the focus point, generally ray under cursor
+	 * @param PointOut computed Focus Point
+	 * @return true if a Focus Point was found, can return false if (eg) the ray missed the target objects
+	 */
+	virtual bool GetWorldSpaceFocusPoint(const FRay& WorldRay, FVector& PointOut) { return false; }
+
+
+};
diff --git a/Engine/Source/Runtime/Experimental/InteractiveToolsFramework/Public/MultiSelectionTool.h b/Engine/Source/Runtime/Experimental/InteractiveToolsFramework/Public/MultiSelectionTool.h
--- a/Engine/Source/Runtime/Experimental/InteractiveToolsFramework/Public/MultiSelectionTool.h
+++ b/Engine/Source/Runtime/Experimental/InteractiveToolsFramework/Public/MultiSelectionTool.h
@@ -5,6 +5,7 @@
 #include "InteractiveTool.h"
 #include "ComponentSourceInterfaces.h"
 #include "ToolTargets/ToolTarget.h"
+#include "InteractiveToolQueryInterfaces.h"
 
 #include "TargetInterfaces/MaterialProvider.h"
 #include "TargetInterfaces/MeshDescriptionCommitter.h"
@@ -14,7 +15,7 @@
 #include "MultiSelectionTool.generated.h"
 
 UCLASS(Transient)
-class INTERACTIVETOOLSFRAMEWORK_API UMultiSelectionTool : public UInteractiveTool
+class INTERACTIVETOOLSFRAMEWORK_API UMultiSelectionTool : public UInteractiveTool, public IInteractiveToolCameraFocusAPI
 {
 GENERATED_BODY()
 public:
@@ -87,6 +88,13 @@ protected:
 	 * Helper to cast a Target into the IMaterialProvider interface.
 	 */
 	IMaterialProvider* TargetMaterialInterface(int32 ComponentIdx) const;
+
+public:
+	// IInteractiveToolCameraFocusAPI implementation
+	virtual bool SupportsWorldSpaceFocusBox() override;
+	virtual FBox GetWorldSpaceFocusBox() override;
+	virtual bool SupportsWorldSpaceFocusPoint() override;
+	virtual bool GetWorldSpaceFocusPoint(const FRay& WorldRay, FVector& PointOut) override;
 };
 
 inline IPrimitiveComponentBackedTarget* UMultiSelectionTool::TargetComponentInterface(int32 TargetIdx) const
diff --git a/Engine/Source/Runtime/Experimental/InteractiveToolsFramework/Public/SingleSelectionTool.h b/Engine/Source/Runtime/Experimental/InteractiveToolsFramework/Public/SingleSelectionTool.h
--- a/Engine/Source/Runtime/Experimental/InteractiveToolsFramework/Public/SingleSelectionTool.h
+++ b/Engine/Source/Runtime/Experimental/InteractiveToolsFramework/Public/SingleSelectionTool.h
@@ -5,11 +5,12 @@
 #include "InteractiveTool.h"
 #include "ComponentSourceInterfaces.h"
 #include "ToolTargets/ToolTarget.h"
+#include "InteractiveToolQueryInterfaces.h"
 
 #include "SingleSelectionTool.generated.h"
 
 UCLASS(Transient)
-class INTERACTIVETOOLSFRAMEWORK_API USingleSelectionTool : public UInteractiveTool
+class INTERACTIVETOOLSFRAMEWORK_API USingleSelectionTool : public UInteractiveTool, public IInteractiveToolCameraFocusAPI
 {
 GENERATED_BODY()
 public:
@@ -36,4 +37,14 @@ public:
 protected:
 	UPROPERTY()
 	TObjectPtr<UToolTarget> Target;
+
+
+
+public:
+	// IInteractiveToolCameraFocusAPI implementation
+	virtual bool SupportsWorldSpaceFocusBox() override;
+	virtual FBox GetWorldSpaceFocusBox() override;
+	virtual bool SupportsWorldSpaceFocusPoint() override;
+	virtual bool GetWorldSpaceFocusPoint(const FRay& WorldRay, FVector& PointOut) override;
+
 };
