MeshModeling: add ModelingObjectsCreationAPI, replaces existing usage of IToolsContextAssetAPI in Modeling Tools/Mode
- add new UModelingObjectsCreationAPI and associated data structures, provides abstract API for creating mesh and texture objects from Tools that is not specifically tied to StaticMesh Actors/Assets
- new helper functions in UE::Modeling:: namespace to simplify usage of an implementation of this API registered in a ContextObjectStore
- add new UEditorModelingObjectsCreationAPI implementation of above, supports hooks for higher level to provide custom paths
- add ModelingModeAssetUtils.h, provides several functions in UE::Modeling:: namespace to be used to implement those hooks (various existing path-determination code is moved here)
- ModelingToolsEditorMode now registers an instance of UEditorModelingObjectsCreationAPI in ContextObjectStore and connects up callbacks to these path functions
- AssetGenerationUtil functions and ModelingModeAssetAPI deleted
- All Tools that previously used IToolsContextAssetAPI updated to use this new system

#rnx

diff --git a/Engine/Plugins/Experimental/MeshLODToolset/Source/MeshLODToolset/Private/Tools/GenerateStaticMeshLODAssetTool.cpp b/Engine/Plugins/Experimental/MeshLODToolset/Source/MeshLODToolset/Private/Tools/GenerateStaticMeshLODAssetTool.cpp
--- a/Engine/Plugins/Experimental/MeshLODToolset/Source/MeshLODToolset/Private/Tools/GenerateStaticMeshLODAssetTool.cpp
+++ b/Engine/Plugins/Experimental/MeshLODToolset/Source/MeshLODToolset/Private/Tools/GenerateStaticMeshLODAssetTool.cpp
@@ -18,7 +18,6 @@
 #include "MeshDescriptionToDynamicMesh.h"
 #include "DynamicMeshToMeshDescription.h"
 
-#include "AssetGenerationUtil.h"
 #include "Selection/ToolSelectionUtil.h"
 
 #include "Physics/PhysicsDataCollection.h"
@@ -164,7 +163,7 @@ const FToolTargetTypeRequirements& UGenerateStaticMeshLODAssetToolBuilder::GetTa
 bool UGenerateStaticMeshLODAssetToolBuilder::CanBuildTool(const FToolBuilderState& SceneState) const
 {
 	// hack to make multi-tool look like single-tool
-	return (AssetAPI != nullptr && SceneState.TargetManager->CountSelectedAndTargetable(SceneState, GetTargetRequirements()) == 1);
+	return (SceneState.TargetManager->CountSelectedAndTargetable(SceneState, GetTargetRequirements()) == 1);
 }
 
 UInteractiveTool* UGenerateStaticMeshLODAssetToolBuilder::BuildTool(const FToolBuilderState& SceneState) const
@@ -174,7 +173,6 @@ UInteractiveTool* UGenerateStaticMeshLODAssetToolBuilder::BuildTool(const FToolB
 	TArray<TObjectPtr<UToolTarget>> Targets = SceneState.TargetManager->BuildAllSelectedTargetable(SceneState, GetTargetRequirements());
 	NewTool->SetTargets(MoveTemp(Targets));
 	NewTool->SetWorld(SceneState.World);
-	NewTool->SetAssetAPI(AssetAPI);
 
 	return NewTool;
 }
@@ -368,11 +366,6 @@ void UGenerateStaticMeshLODAssetTool::Shutdown(EToolShutdownType ShutdownType)
 	FDynamicMeshOpResult Result = PreviewWithBackgroundCompute->Shutdown();
 }
 
-void UGenerateStaticMeshLODAssetTool::SetAssetAPI(IAssetGenerationAPI* AssetAPIIn)
-{
-	this->AssetAPI = AssetAPIIn;
-}
-
 
 bool UGenerateStaticMeshLODAssetTool::CanAccept() const
 {
diff --git a/Engine/Plugins/Experimental/MeshLODToolset/Source/MeshLODToolset/Public/Tools/GenerateStaticMeshLODAssetTool.h b/Engine/Plugins/Experimental/MeshLODToolset/Source/MeshLODToolset/Public/Tools/GenerateStaticMeshLODAssetTool.h
--- a/Engine/Plugins/Experimental/MeshLODToolset/Source/MeshLODToolset/Public/Tools/GenerateStaticMeshLODAssetTool.h
+++ b/Engine/Plugins/Experimental/MeshLODToolset/Source/MeshLODToolset/Public/Tools/GenerateStaticMeshLODAssetTool.h
@@ -18,7 +18,6 @@
 // predeclarations
 struct FMeshDescription;
 class USimpleDynamicMeshComponent;
-class IAssetGenerationAPI;
 class UMeshOpPreviewWithBackgroundCompute;
 class UGenerateStaticMeshLODAssetTool;
 namespace GenerateStaticMeshLODAssetLocals
@@ -43,8 +42,6 @@ class MESHLODTOOLSET_API UGenerateStaticMeshLODAssetToolBuilder : public UIntera
 	GENERATED_BODY()
 
 public:
-	IAssetGenerationAPI* AssetAPI = nullptr;
-
 	virtual bool CanBuildTool(const FToolBuilderState& SceneState) const override;
 	virtual UInteractiveTool* BuildTool(const FToolBuilderState& SceneState) const override;
 
@@ -144,7 +141,6 @@ public:
 	virtual void Shutdown(EToolShutdownType ShutdownType) override;
 
 	virtual void SetWorld(UWorld* World);
-	virtual void SetAssetAPI(IAssetGenerationAPI* AssetAPI);
 
 	virtual void OnTick(float DeltaTime);
 
@@ -185,7 +181,6 @@ protected:
 
 protected:
 	UWorld* TargetWorld;
-	IAssetGenerationAPI* AssetAPI;
 
 	UPROPERTY()
 	UGenerateStaticMeshLODProcess* GenerateProcess;
diff --git a/Engine/Plugins/Experimental/MeshLODToolset/Source/MeshLODToolset/Public/Tools/LODManagerTool.h b/Engine/Plugins/Experimental/MeshLODToolset/Source/MeshLODToolset/Public/Tools/LODManagerTool.h
--- a/Engine/Plugins/Experimental/MeshLODToolset/Source/MeshLODToolset/Public/Tools/LODManagerTool.h
+++ b/Engine/Plugins/Experimental/MeshLODToolset/Source/MeshLODToolset/Public/Tools/LODManagerTool.h
@@ -12,7 +12,6 @@
 #include "LODManagerTool.generated.h"
 
 // predeclarations
-class IAssetGenerationAPI;
 class UMaterialInterface;
 
 /**
@@ -24,8 +23,6 @@ class MESHLODTOOLSET_API ULODManagerToolBuilder : public UInteractiveToolBuilder
 	GENERATED_BODY()
 
 public:
-	IAssetGenerationAPI* AssetAPI = nullptr;
-
 	virtual bool CanBuildTool(const FToolBuilderState& SceneState) const override;
 	virtual UInteractiveTool* BuildTool(const FToolBuilderState& SceneState) const override;
 
diff --git a/Engine/Plugins/Experimental/MeshModelingToolset/Source/MeshModelingTools/Private/AddPatchTool.cpp b/Engine/Plugins/Experimental/MeshModelingToolset/Source/MeshModelingTools/Private/AddPatchTool.cpp
--- a/Engine/Plugins/Experimental/MeshModelingToolset/Source/MeshModelingTools/Private/AddPatchTool.cpp
+++ b/Engine/Plugins/Experimental/MeshModelingToolset/Source/MeshModelingTools/Private/AddPatchTool.cpp
@@ -5,7 +5,7 @@
 #include "InteractiveToolManager.h"
 #include "BaseBehaviors/MouseHoverBehavior.h"
 #include "Selection/ToolSelectionUtil.h"
-#include "AssetGenerationUtil.h"
+#include "ModelingObjectsCreationAPI.h"
 #include "ToolSceneQueriesUtil.h"
 
 #include "MeshDescriptionBuilder.h"
@@ -28,14 +28,13 @@ using namespace UE::Geometry;
  */
 bool UAddPatchToolBuilder::CanBuildTool(const FToolBuilderState& SceneState) const
 {
-	return (this->AssetAPI != nullptr);
+	return true;
 }
 
 UInteractiveTool* UAddPatchToolBuilder::BuildTool(const FToolBuilderState& SceneState) const
 {
 	UAddPatchTool* NewTool = NewObject<UAddPatchTool>(SceneState.ToolManager);
 	NewTool->SetWorld(SceneState.World);
-	NewTool->SetAssetAPI(AssetAPI);
 	return NewTool;
 }
 
@@ -55,12 +54,6 @@ void UAddPatchTool::SetWorld(UWorld* World)
 	this->TargetWorld = World;
 }
 
-void UAddPatchTool::SetAssetAPI(IAssetGenerationAPI* AssetAPIIn)
-{
-	this->AssetAPI = AssetAPIIn;
-}
-
-
 void UAddPatchTool::Setup()
 {
 	USingleClickTool::Setup();
@@ -298,26 +291,25 @@ void UAddPatchTool::GeneratePreviewBaseMesh()
 
 void UAddPatchTool::OnClicked(const FInputDeviceRay& DeviceClickPos)
 {
-#if WITH_EDITOR
 	const FDynamicMesh3* CurMesh = PreviewMesh->GetPreviewDynamicMesh();
 	UE::Geometry::FTransform3d CurTransform(PreviewMesh->GetTransform());
 	UMaterialInterface* Material = PreviewMesh->GetMaterial();
 	GetToolManager()->BeginUndoTransaction(LOCTEXT("AddPatchToolTransactionName", "Add Patch Mesh"));
 
-	AActor* NewActor = AssetGenerationUtil::GenerateStaticMeshActor(
-		AssetAPI, TargetWorld,
-		CurMesh, CurTransform, TEXT("Patch"), Material );
-
-	// select newly-created object
-	if (NewActor != nullptr)
+	FCreateMeshObjectParams NewMeshObjectParams;
+	NewMeshObjectParams.TargetWorld = TargetWorld;
+	NewMeshObjectParams.Transform = (FTransform)CurTransform;
+	NewMeshObjectParams.BaseName = TEXT("Patch");
+	NewMeshObjectParams.Materials.Add(Material);
+	NewMeshObjectParams.SetMesh(CurMesh);
+	FCreateMeshObjectResult Result = UE::Modeling::CreateMeshObject(GetToolManager(), MoveTemp(NewMeshObjectParams));
+	if (Result.IsOK() && Result.NewActor != nullptr)
 	{
-		ToolSelectionUtil::SetNewActorSelection(GetToolManager(), NewActor);
+		// select newly-created object
+		ToolSelectionUtil::SetNewActorSelection(GetToolManager(), Result.NewActor);
 	}
 
 	GetToolManager()->EndUndoTransaction();
-#else
-	checkNoEntry();
-#endif
 }
 
 
diff --git a/Engine/Plugins/Experimental/MeshModelingToolset/Source/MeshModelingTools/Private/AddPrimitiveTool.cpp b/Engine/Plugins/Experimental/MeshModelingToolset/Source/MeshModelingTools/Private/AddPrimitiveTool.cpp
--- a/Engine/Plugins/Experimental/MeshModelingToolset/Source/MeshModelingTools/Private/AddPrimitiveTool.cpp
+++ b/Engine/Plugins/Experimental/MeshModelingToolset/Source/MeshModelingTools/Private/AddPrimitiveTool.cpp
@@ -5,7 +5,7 @@
 #include "InteractiveToolManager.h"
 #include "BaseBehaviors/MouseHoverBehavior.h"
 #include "Selection/ToolSelectionUtil.h"
-#include "AssetGenerationUtil.h"
+#include "ModelingObjectsCreationAPI.h"
 #include "ToolSceneQueriesUtil.h"
 
 #include "Components/StaticMeshComponent.h"
@@ -38,7 +38,7 @@ using namespace UE::Geometry;
  */
 bool UAddPrimitiveToolBuilder::CanBuildTool(const FToolBuilderState& SceneState) const
 {
-	return (this->AssetAPI != nullptr);
+	return true;
 }
 
 UInteractiveTool* UAddPrimitiveToolBuilder::BuildTool(const FToolBuilderState& SceneState) const
@@ -77,7 +77,6 @@ UInteractiveTool* UAddPrimitiveToolBuilder::BuildTool(const FToolBuilderState& S
 		break;
 	}
 	NewTool->SetWorld(SceneState.World);
-	NewTool->SetAssetAPI(AssetAPI);
 	return NewTool;
 }
 
@@ -109,11 +108,6 @@ void UAddPrimitiveTool::SetWorld(UWorld* World)
 	this->TargetWorld = World;
 }
 
-void UAddPrimitiveTool::SetAssetAPI(IAssetGenerationAPI* AssetAPIIn)
-{
-	this->AssetAPI = AssetAPIIn;
-}
-
 UAddPrimitiveTool::UAddPrimitiveTool(const FObjectInitializer&)
 {
 	ShapeSettings = CreateDefaultSubobject<UProceduralShapeToolProperties>(TEXT("ShapeSettings"));
@@ -316,10 +310,9 @@ void UAddPrimitiveTool::UpdatePreviewMesh()
 
 void UAddPrimitiveTool::OnClicked(const FInputDeviceRay& DeviceClickPos)
 {
-#if WITH_EDITOR
 	UMaterialInterface* Material = PreviewMesh->GetMaterial();
 
-	if (ShapeSettings->bInstanceIfPossible && IsEquivalentLastGeneratedAsset())
+	if (ShapeSettings->bInstanceIfPossible && LastGenerated != nullptr && IsEquivalentLastGeneratedAsset())
 	{
 		GetToolManager()->BeginUndoTransaction(LOCTEXT("AddPrimitiveToolTransactionName", "Add Primitive Mesh"));
 		FActorSpawnParameters SpawnParameters;
@@ -331,39 +324,51 @@ void UAddPrimitiveTool::OnClicked(const FInputDeviceRay& DeviceClickPos)
 		CloneActor->GetStaticMeshComponent()->SetWorldTransform(PreviewMesh->GetTransform());
 		CloneActor->GetStaticMeshComponent()->SetStaticMesh(LastGenerated->StaticMesh);
 		CloneActor->GetStaticMeshComponent()->SetMaterial(0, Material);
+#if WITH_EDITOR
 		CloneActor->SetActorLabel(LastGenerated->Label);
+#endif
 		// select newly-created object
 		ToolSelectionUtil::SetNewActorSelection(GetToolManager(), CloneActor);
 		GetToolManager()->EndUndoTransaction();
 
 		return;
 	}
+	LastGenerated = nullptr;
 
 	const FDynamicMesh3* CurMesh = PreviewMesh->GetPreviewDynamicMesh();
 	UE::Geometry::FTransform3d CurTransform(PreviewMesh->GetTransform());
 
 	GetToolManager()->BeginUndoTransaction(LOCTEXT("AddPrimitiveToolTransactionName", "Add Primitive Mesh"));
 
-	AActor* NewActor = AssetGenerationUtil::GenerateStaticMeshActor(
-		AssetAPI, TargetWorld,
-		CurMesh, CurTransform, AssetName, Material);
-	if (NewActor != nullptr)
+	FCreateMeshObjectParams NewMeshObjectParams;
+	NewMeshObjectParams.TargetWorld = TargetWorld;
+	NewMeshObjectParams.Transform = (FTransform)CurTransform;
+	NewMeshObjectParams.BaseName = AssetName;
+	NewMeshObjectParams.Materials.Add(Material);
+	NewMeshObjectParams.SetMesh(CurMesh);
+	FCreateMeshObjectResult Result = UE::Modeling::CreateMeshObject(GetToolManager(), MoveTemp(NewMeshObjectParams));
+	if (Result.IsOK() )
 	{
-		LastGenerated = NewObject<ULastActorInfo>(this);
-		LastGenerated->ShapeSettings = DuplicateObject(ShapeSettings, nullptr);
-		LastGenerated->MaterialProperties = DuplicateObject(MaterialProperties, nullptr);
-		LastGenerated->Actor = NewActor;
-		LastGenerated->StaticMesh = CastChecked<AStaticMeshActor>(LastGenerated->Actor)->GetStaticMeshComponent()->GetStaticMesh();
-		LastGenerated->Label = LastGenerated->Actor->GetActorLabel();
+		if (Result.NewActor != nullptr)
+		{
+			if (Cast<AStaticMeshActor>(Result.NewActor) != nullptr)
+			{
+				LastGenerated = NewObject<ULastActorInfo>(this);
+				LastGenerated->ShapeSettings = DuplicateObject(ShapeSettings, nullptr);
+				LastGenerated->MaterialProperties = DuplicateObject(MaterialProperties, nullptr);
+				LastGenerated->Actor = Result.NewActor;
+				LastGenerated->StaticMesh = CastChecked<AStaticMeshActor>(LastGenerated->Actor)->GetStaticMeshComponent()->GetStaticMesh();
+#if WITH_EDITOR
+				LastGenerated->Label = LastGenerated->Actor->GetActorLabel();
+#endif
+			}
 
-		// select newly-created object
-		ToolSelectionUtil::SetNewActorSelection(GetToolManager(), LastGenerated->Actor);
+			// select newly-created object
+			ToolSelectionUtil::SetNewActorSelection(GetToolManager(), Result.NewActor);
+		}
 	}
 
 	GetToolManager()->EndUndoTransaction();
-#else
-	checkNoEntry();
-#endif
 }
 
 
diff --git a/Engine/Plugins/Experimental/MeshModelingToolset/Source/MeshModelingTools/Private/BakeMeshAttributeMapsTool.cpp b/Engine/Plugins/Experimental/MeshModelingToolset/Source/MeshModelingTools/Private/BakeMeshAttributeMapsTool.cpp
--- a/Engine/Plugins/Experimental/MeshModelingToolset/Source/MeshModelingTools/Private/BakeMeshAttributeMapsTool.cpp
+++ b/Engine/Plugins/Experimental/MeshModelingToolset/Source/MeshModelingTools/Private/BakeMeshAttributeMapsTool.cpp
@@ -23,7 +23,7 @@
 
 #include "AssetUtils/Texture2DBuilder.h"
 #include "AssetUtils/MeshDescriptionUtil.h"
-#include "AssetGenerationUtil.h"
+#include "ModelingObjectsCreationAPI.h"
 
 #include "TargetInterfaces/MaterialProvider.h"
 #include "TargetInterfaces/MeshDescriptionProvider.h"
@@ -67,7 +67,6 @@ UInteractiveTool* UBakeMeshAttributeMapsToolBuilder::BuildTool(const FToolBuilde
 
 	TArray<TObjectPtr<UToolTarget>> Targets = SceneState.TargetManager->BuildAllSelectedTargetable(SceneState, GetTargetRequirements());
 	NewTool->SetTargets(MoveTemp(Targets));
-	NewTool->SetAssetAPI(AssetAPI);
 
 	return NewTool;
 }
@@ -345,11 +344,6 @@ public:
  * Tool
  */
 
-void UBakeMeshAttributeMapsTool::SetAssetAPI(IAssetGenerationAPI* AssetAPIIn)
-{
-	AssetAPI = AssetAPIIn;
-}
-
 void UBakeMeshAttributeMapsTool::Setup()
 {
 	UInteractiveTool::Setup();
@@ -615,74 +609,70 @@ void UBakeMeshAttributeMapsTool::Shutdown(EToolShutdownType ShutdownType)
 			check(StaticMeshAsset);
 			FString BaseName = TargetComponent->GetOwnerActor()->GetName();
 
-			if (AssetAPI != nullptr)
+			bool bCreatedAssetOK = false;
+			switch (Settings->MapType)
 			{
-				bool bCreatedAssetOK = false;
-				switch (Settings->MapType)
-				{
-				default:
-					check(false);
-					break;
-
-				case EBakeMapType::TangentSpaceNormalMap:
-					FTexture2DBuilder::CopyPlatformDataToSourceData(Settings->Result[0], FTexture2DBuilder::ETextureType::NormalMap);
-					bCreatedAssetOK = AssetGenerationUtil::SaveGeneratedTexture2D(AssetAPI, Settings->Result[0],
-						FString::Printf(TEXT("%s_Normals"), *BaseName), StaticMeshAsset);
-					break;
-
-				case EBakeMapType::Occlusion:
-					FTexture2DBuilder::CopyPlatformDataToSourceData(Settings->Result[0], FTexture2DBuilder::ETextureType::AmbientOcclusion);
-					bCreatedAssetOK = AssetGenerationUtil::SaveGeneratedTexture2D(AssetAPI, Settings->Result[0],
-						FString::Printf(TEXT("%s_Occlusion"), *BaseName), StaticMeshAsset);
-					FTexture2DBuilder::CopyPlatformDataToSourceData(Settings->Result[1], FTexture2DBuilder::ETextureType::NormalMap);
-					bCreatedAssetOK = AssetGenerationUtil::SaveGeneratedTexture2D(AssetAPI, Settings->Result[1],
-						FString::Printf(TEXT("%s_BentNormal"), *BaseName), StaticMeshAsset);
-					break;
-
-				case EBakeMapType::Curvature:
-					FTexture2DBuilder::CopyPlatformDataToSourceData(Settings->Result[0], FTexture2DBuilder::ETextureType::Color);
-					bCreatedAssetOK = AssetGenerationUtil::SaveGeneratedTexture2D(AssetAPI, Settings->Result[0],
-						FString::Printf(TEXT("%s_Curvature"), *BaseName), StaticMeshAsset);
-					break;
-
-				case EBakeMapType::NormalImage:
-					FTexture2DBuilder::CopyPlatformDataToSourceData(Settings->Result[0], FTexture2DBuilder::ETextureType::Color);
-					bCreatedAssetOK = AssetGenerationUtil::SaveGeneratedTexture2D(AssetAPI, Settings->Result[0],
-						FString::Printf(TEXT("%s_NormalImg"), *BaseName), StaticMeshAsset);
-					break;
-
-				case EBakeMapType::FaceNormalImage:
-					FTexture2DBuilder::CopyPlatformDataToSourceData(Settings->Result[0], FTexture2DBuilder::ETextureType::Color);
-					bCreatedAssetOK = AssetGenerationUtil::SaveGeneratedTexture2D(AssetAPI, Settings->Result[0],
-						FString::Printf(TEXT("%s_FaceNormalImg"), *BaseName), StaticMeshAsset);
-					break;
-				case EBakeMapType::MaterialID:
-					FTexture2DBuilder::CopyPlatformDataToSourceData(Settings->Result[0], FTexture2DBuilder::ETextureType::Color);
-					bCreatedAssetOK = AssetGenerationUtil::SaveGeneratedTexture2D(AssetAPI, Settings->Result[0],
-																				  FString::Printf(TEXT("%s_MaterialIDImg"), *BaseName), StaticMeshAsset);
-					break;
-				case EBakeMapType::PositionImage:
-					FTexture2DBuilder::CopyPlatformDataToSourceData(Settings->Result[0], FTexture2DBuilder::ETextureType::Color);
-					bCreatedAssetOK = AssetGenerationUtil::SaveGeneratedTexture2D(AssetAPI, Settings->Result[0],
-						FString::Printf(TEXT("%s_PositionImg"), *BaseName), StaticMeshAsset);
-					break;
-
-				case EBakeMapType::Texture2DImage:
-					FTexture2DBuilder::CopyPlatformDataToSourceData(Settings->Result[0], FTexture2DBuilder::ETextureType::Color);
-					bCreatedAssetOK = AssetGenerationUtil::SaveGeneratedTexture2D(AssetAPI, Settings->Result[0],
-						FString::Printf(TEXT("%s_TextureImg"), *BaseName), StaticMeshAsset);
-					break;
-
-				case EBakeMapType::MultiTexture:
-					FTexture2DBuilder::CopyPlatformDataToSourceData(Settings->Result[0], FTexture2DBuilder::ETextureType::Color);
-					bCreatedAssetOK = AssetGenerationUtil::SaveGeneratedTexture2D(AssetAPI, Settings->Result[0],
-																				  FString::Printf(TEXT("%s_MultiTextureImg"), *BaseName), StaticMeshAsset);
-					break;
-
-				}
-				ensure(bCreatedAssetOK);
-			}
+			default:
+				check(false);
+				break;
+
+			case EBakeMapType::TangentSpaceNormalMap:
+				FTexture2DBuilder::CopyPlatformDataToSourceData(Settings->Result[0], FTexture2DBuilder::ETextureType::NormalMap);
+				bCreatedAssetOK = UE::Modeling::CreateTextureObject(GetToolManager(),
+					FCreateTextureObjectParams{ 0, StaticMeshAsset->GetWorld(), StaticMeshAsset, FString::Printf(TEXT("%s_Normals"), *BaseName), Settings->Result[0] }).IsOK();
+				break;
+
+			case EBakeMapType::Occlusion:
+				FTexture2DBuilder::CopyPlatformDataToSourceData(Settings->Result[0], FTexture2DBuilder::ETextureType::AmbientOcclusion);
+				bCreatedAssetOK = UE::Modeling::CreateTextureObject(GetToolManager(),
+					FCreateTextureObjectParams{ 0, StaticMeshAsset->GetWorld(), StaticMeshAsset, FString::Printf(TEXT("%s_Occlusion"), *BaseName), Settings->Result[0] }).IsOK();
+				FTexture2DBuilder::CopyPlatformDataToSourceData(Settings->Result[1], FTexture2DBuilder::ETextureType::NormalMap);
+				bCreatedAssetOK = UE::Modeling::CreateTextureObject(GetToolManager(),
+					FCreateTextureObjectParams{ 0, StaticMeshAsset->GetWorld(), StaticMeshAsset, FString::Printf(TEXT("%s_BentNormal"), *BaseName), Settings->Result[1] }).IsOK();
+				break;
+
+			case EBakeMapType::Curvature:
+				FTexture2DBuilder::CopyPlatformDataToSourceData(Settings->Result[0], FTexture2DBuilder::ETextureType::Color);
+				bCreatedAssetOK = UE::Modeling::CreateTextureObject(GetToolManager(),
+					FCreateTextureObjectParams{ 0, StaticMeshAsset->GetWorld(), StaticMeshAsset, FString::Printf(TEXT("%s_Curvature"), *BaseName), Settings->Result[0] }).IsOK();
+				break;
+
+			case EBakeMapType::NormalImage:
+				FTexture2DBuilder::CopyPlatformDataToSourceData(Settings->Result[0], FTexture2DBuilder::ETextureType::Color);
+				bCreatedAssetOK = UE::Modeling::CreateTextureObject(GetToolManager(),
+					FCreateTextureObjectParams{ 0, StaticMeshAsset->GetWorld(), StaticMeshAsset, FString::Printf(TEXT("%s_NormalImg"), *BaseName), Settings->Result[0] }).IsOK();
+				break;
+
+			case EBakeMapType::FaceNormalImage:
+				FTexture2DBuilder::CopyPlatformDataToSourceData(Settings->Result[0], FTexture2DBuilder::ETextureType::Color);
+				bCreatedAssetOK = UE::Modeling::CreateTextureObject(GetToolManager(),
+					FCreateTextureObjectParams{ 0, StaticMeshAsset->GetWorld(), StaticMeshAsset, FString::Printf(TEXT("%s_FaceNormalImg"), *BaseName), Settings->Result[0] }).IsOK();
+				break;
+			case EBakeMapType::MaterialID:
+				FTexture2DBuilder::CopyPlatformDataToSourceData(Settings->Result[0], FTexture2DBuilder::ETextureType::Color);
+				bCreatedAssetOK = UE::Modeling::CreateTextureObject(GetToolManager(),
+					FCreateTextureObjectParams{ 0, StaticMeshAsset->GetWorld(), StaticMeshAsset, FString::Printf(TEXT("%s_MaterialIDImg"), *BaseName), Settings->Result[0] }).IsOK();
+				break;
+			case EBakeMapType::PositionImage:
+				FTexture2DBuilder::CopyPlatformDataToSourceData(Settings->Result[0], FTexture2DBuilder::ETextureType::Color);
+				bCreatedAssetOK = UE::Modeling::CreateTextureObject(GetToolManager(),
+					FCreateTextureObjectParams{ 0, StaticMeshAsset->GetWorld(), StaticMeshAsset, FString::Printf(TEXT("%s_PositionImg"), *BaseName), Settings->Result[0] }).IsOK();
+				break;
+
+			case EBakeMapType::Texture2DImage:
+				FTexture2DBuilder::CopyPlatformDataToSourceData(Settings->Result[0], FTexture2DBuilder::ETextureType::Color);
+				bCreatedAssetOK = UE::Modeling::CreateTextureObject(GetToolManager(),
+					FCreateTextureObjectParams{ 0, StaticMeshAsset->GetWorld(), StaticMeshAsset, FString::Printf(TEXT("%s_TextureImg"), *BaseName), Settings->Result[0] }).IsOK();
+				break;
+
+			case EBakeMapType::MultiTexture:
+				FTexture2DBuilder::CopyPlatformDataToSourceData(Settings->Result[0], FTexture2DBuilder::ETextureType::Color);
+				bCreatedAssetOK = UE::Modeling::CreateTextureObject(GetToolManager(),
+					FCreateTextureObjectParams{ 0, StaticMeshAsset->GetWorld(), StaticMeshAsset, FString::Printf(TEXT("%s_MultiTextureImg"), *BaseName), Settings->Result[0] }).IsOK();
+				break;
 
+			}
+			ensure(bCreatedAssetOK);
 		}
 
 		DynamicMeshComponent->UnregisterComponent();
diff --git a/Engine/Plugins/Experimental/MeshModelingToolset/Source/MeshModelingTools/Private/BakeTransformTool.cpp b/Engine/Plugins/Experimental/MeshModelingToolset/Source/MeshModelingTools/Private/BakeTransformTool.cpp
--- a/Engine/Plugins/Experimental/MeshModelingToolset/Source/MeshModelingTools/Private/BakeTransformTool.cpp
+++ b/Engine/Plugins/Experimental/MeshModelingToolset/Source/MeshModelingTools/Private/BakeTransformTool.cpp
@@ -16,7 +16,6 @@
 #include "MeshDescriptionToDynamicMesh.h"
 #include "DynamicMeshToMeshDescription.h"
 
-#include "AssetGenerationUtil.h"
 #include "Physics/ComponentCollisionUtil.h"
 
 #include "TargetInterfaces/MeshDescriptionCommitter.h"
@@ -58,7 +57,6 @@ UInteractiveTool* UBakeTransformToolBuilder::BuildTool(const FToolBuilderState&
 	TArray<TObjectPtr<UToolTarget>> Targets = SceneState.TargetManager->BuildAllSelectedTargetable(SceneState, GetTargetRequirements());
 	NewTool->SetTargets(MoveTemp(Targets));
 	NewTool->SetWorld(SceneState.World);
-	NewTool->SetAssetAPI(AssetAPI);
 
 	return NewTool;
 }
@@ -133,12 +131,6 @@ void UBakeTransformTool::Shutdown(EToolShutdownType ShutdownType)
 	}
 }
 
-void UBakeTransformTool::SetAssetAPI(IToolsContextAssetAPI* AssetAPIIn)
-{
-	this->AssetAPI = AssetAPIIn;
-}
-
-
 
 
 void UBakeTransformTool::UpdateAssets()
diff --git a/Engine/Plugins/Experimental/MeshModelingToolset/Source/MeshModelingTools/Private/CombineMeshesTool.cpp b/Engine/Plugins/Experimental/MeshModelingToolset/Source/MeshModelingTools/Private/CombineMeshesTool.cpp
--- a/Engine/Plugins/Experimental/MeshModelingToolset/Source/MeshModelingTools/Private/CombineMeshesTool.cpp
+++ b/Engine/Plugins/Experimental/MeshModelingToolset/Source/MeshModelingTools/Private/CombineMeshesTool.cpp
@@ -15,7 +15,7 @@
 #include "MeshDescriptionToDynamicMesh.h"
 #include "DynamicMeshToMeshDescription.h"
 
-#include "AssetGenerationUtil.h"
+#include "ModelingObjectsCreationAPI.h"
 #include "Selection/ToolSelectionUtil.h"
 #include "Physics/ComponentCollisionUtil.h"
 #include "ShapeApproximation/SimpleShapeSet3.h"
@@ -57,8 +57,8 @@ const FToolTargetTypeRequirements& UCombineMeshesToolBuilder::GetTargetRequireme
 bool UCombineMeshesToolBuilder::CanBuildTool(const FToolBuilderState& SceneState) const
 {
 	return (bIsDuplicateTool) ?
-		  (AssetAPI != nullptr && SceneState.TargetManager->CountSelectedAndTargetable(SceneState, GetTargetRequirements()) == 1)
-		: (AssetAPI != nullptr && SceneState.TargetManager->CountSelectedAndTargetable(SceneState, GetTargetRequirements()) > 1);
+		  (SceneState.TargetManager->CountSelectedAndTargetable(SceneState, GetTargetRequirements()) == 1)
+		: (SceneState.TargetManager->CountSelectedAndTargetable(SceneState, GetTargetRequirements()) > 1);
 }
 
 UInteractiveTool* UCombineMeshesToolBuilder::BuildTool(const FToolBuilderState& SceneState) const
@@ -68,7 +68,6 @@ UInteractiveTool* UCombineMeshesToolBuilder::BuildTool(const FToolBuilderState&
 	TArray<TObjectPtr<UToolTarget>> Targets = SceneState.TargetManager->BuildAllSelectedTargetable(SceneState, GetTargetRequirements());
 	NewTool->SetTargets(MoveTemp(Targets));
 	NewTool->SetWorld(SceneState.World);
-	NewTool->SetAssetAPI(AssetAPI);
 	NewTool->SetDuplicateMode(bIsDuplicateTool);
 
 	return NewTool;
@@ -109,14 +108,14 @@ void UCombineMeshesTool::Setup()
 		else
 		{
 			int32 Index = (BasicProperties->WriteOutputTo == ECombineTargetType::FirstInputAsset) ? 0 : Targets.Num() - 1;
-			BasicProperties->OutputAsset = AssetGenerationUtil::GetComponentAssetBaseName(TargetComponentInterface(Index)->GetOwnerComponent(), false);
+			BasicProperties->OutputAsset = UE::Modeling::GetComponentAssetBaseName(TargetComponentInterface(Index)->GetOwnerComponent(), false);
 		}
 	});
 
 	if (bDuplicateMode)
 	{
 		SetToolDisplayName(LOCTEXT("DuplicateMeshesToolName", "Duplicate"));
-		BasicProperties->OutputName = AssetGenerationUtil::GetComponentAssetBaseName(TargetComponentInterface(0)->GetOwnerComponent());
+		BasicProperties->OutputName = UE::Modeling::GetComponentAssetBaseName(TargetComponentInterface(0)->GetOwnerComponent());
 	}
 	else
 	{
@@ -164,13 +163,6 @@ void UCombineMeshesTool::Shutdown(EToolShutdownType ShutdownType)
 	}
 }
 
-void UCombineMeshesTool::SetAssetAPI(IAssetGenerationAPI* AssetAPIIn)
-{
-	this->AssetAPI = AssetAPIIn;
-}
-
-
-
 void UCombineMeshesTool::CreateNewAsset()
 {
 	// Make sure mesh descriptions are deserialized before we open transaction.
@@ -189,8 +181,6 @@ void UCombineMeshesTool::CreateNewAsset()
 
 	// note there is a MergeMeshUtilities.h w/ a very feature-filled mesh merging class, but for simplicity (and to fit modeling tool needs)
 	// this tool currently converts things through dynamic mesh instead
-
-#if WITH_EDITOR
 	FBox Box(ForceInit);
 	for (int32 ComponentIdx = 0; ComponentIdx < Targets.Num(); ComponentIdx++)
 	{
@@ -237,16 +227,20 @@ void UCombineMeshesTool::CreateNewAsset()
 	UE::Geometry::FComponentCollisionSettings CollisionSettings;
 
 	{
+#if WITH_EDITOR
 		FScopedSlowTask SlowTask(Targets.Num()+1, 
 			bDuplicateMode ? 
 			LOCTEXT("DuplicateMeshBuild", "Building duplicate mesh ...") :
 			LOCTEXT("CombineMeshesBuild", "Building combined mesh ..."));
 		SlowTask.MakeDialog();
+#endif
 		bool bNeedColorAttr = false;
 		int MatIndexBase = 0;
 		for (int32 ComponentIdx = 0; ComponentIdx < Targets.Num(); ComponentIdx++)
 		{
+#if WITH_EDITOR
 			SlowTask.EnterProgressFrame(1);
+#endif
 			IPrimitiveComponentBackedTarget* TargetComponent = TargetComponentInterface(ComponentIdx);
 			IMeshDescriptionProvider* TargetMeshProvider = TargetMeshProviderInterface(ComponentIdx);
 
@@ -298,7 +292,9 @@ void UCombineMeshesTool::CreateNewAsset()
 			AccumulateDMesh.Attributes()->DisablePrimaryColors();
 		}
 
+#if WITH_EDITOR
 		SlowTask.EnterProgressFrame(1);
+#endif
 
 		if (bDuplicateMode)
 		{
@@ -313,32 +309,40 @@ void UCombineMeshesTool::CreateNewAsset()
 		{
 			UseBaseName = (bDuplicateMode) ? TEXT("Duplicate") : TEXT("Combined");
 		}
-		AActor* NewActor = AssetGenerationUtil::GenerateStaticMeshActor(
-			AssetAPI, TargetWorld,
-			&AccumulateDMesh, (UE::Geometry::FTransform3d)AccumToWorld, UseBaseName, AllMaterials);
-		if (NewActor != nullptr)
+
+		FCreateMeshObjectParams NewMeshObjectParams;
+		NewMeshObjectParams.TargetWorld = TargetWorld;
+		NewMeshObjectParams.Transform = AccumToWorld;
+		NewMeshObjectParams.BaseName = UseBaseName;
+		NewMeshObjectParams.Materials = AllMaterials;
+		NewMeshObjectParams.SetMesh(&AccumulateDMesh);
+		FCreateMeshObjectResult Result = UE::Modeling::CreateMeshObject(GetToolManager(), MoveTemp(NewMeshObjectParams));
+		if (Result.IsOK() && Result.NewActor != nullptr)
 		{
 			// copy the component materials onto the new static mesh asset too
 			// (note: GenerateStaticMeshActor defaults to just putting blank slots on the asset)
-			UStaticMeshComponent* NewMeshComponent = NewActor->FindComponentByClass<UStaticMeshComponent>();
-			UStaticMesh* NewMesh = NewMeshComponent->GetStaticMesh();
-			for (int32 MatIdx = 0; MatIdx < AllMaterials.Num(); MatIdx++)
+			UStaticMeshComponent* NewMeshComponent = Result.NewActor->FindComponentByClass<UStaticMeshComponent>();
+			if (NewMeshComponent)
 			{
-				NewMesh->SetMaterial(MatIdx, AllMaterials[MatIdx]);
-			}
+				UStaticMesh* NewMesh = NewMeshComponent->GetStaticMesh();
+#if WITH_EDITOR
+				for (int32 MatIdx = 0; MatIdx < AllMaterials.Num(); MatIdx++)
+				{
+					NewMesh->SetMaterial(MatIdx, AllMaterials[MatIdx]);
+				}
+#endif
 
-			// if any inputs have Simple Collision geometry we will forward it to new mesh.
-			if (SimpleCollision.TotalElementsNum() > 0)
-			{
-				UE::Geometry::SetSimpleCollision(NewMeshComponent, &SimpleCollision, CollisionSettings);
+				// if any inputs have Simple Collision geometry we will forward it to new mesh.
+				if (SimpleCollision.TotalElementsNum() > 0)
+				{
+					UE::Geometry::SetSimpleCollision(NewMeshComponent, &SimpleCollision, CollisionSettings);
+				}
 			}
 
 			// select the new actor
-			ToolSelectionUtil::SetNewActorSelection(GetToolManager(), NewActor);
+			ToolSelectionUtil::SetNewActorSelection(GetToolManager(), Result.NewActor);
 		}
 	}
-#endif
-
 	
 	TArray<AActor*> Actors;
 	for (int32 Idx = 0; Idx < Targets.Num(); Idx++)
@@ -376,8 +380,6 @@ void UCombineMeshesTool::UpdateExistingAsset()
 
 	AActor* SkipActor = nullptr;
 
-#if WITH_EDITOR
-
 	bool bMergeSameMaterials = true;
 	TArray<UMaterialInterface*> AllMaterials;
 	TMap<UMaterialInterface*, int> KnownMaterials;
@@ -427,16 +429,20 @@ void UCombineMeshesTool::UpdateExistingAsset()
 	Transforms.SetNum(2);
 
 	{
+#if WITH_EDITOR
 		FScopedSlowTask SlowTask(Targets.Num()+1, 
 			bDuplicateMode ? 
 			LOCTEXT("DuplicateMeshBuild", "Building duplicate mesh ...") :
 			LOCTEXT("CombineMeshesBuild", "Building combined mesh ..."));
 		SlowTask.MakeDialog();
+#endif
 		bool bNeedColorAttr = false;
 		int MatIndexBase = 0;
 		for (int32 ComponentIdx = 0; ComponentIdx < Targets.Num(); ComponentIdx++)
 		{
+#if WITH_EDITOR
 			SlowTask.EnterProgressFrame(1);
+#endif
 
 			IPrimitiveComponentBackedTarget* TargetComponent = TargetComponentInterface(ComponentIdx);
 
@@ -486,7 +492,9 @@ void UCombineMeshesTool::UpdateExistingAsset()
 			AccumulateDMesh.Attributes()->DisablePrimaryColors();
 		}
 
+#if WITH_EDITOR
 		SlowTask.EnterProgressFrame(1);
+#endif
 
 		UpdateTargetCommitter->CommitMeshDescription([&](const IMeshDescriptionCommitter::FCommitterParams& CommitParams)
 		{
@@ -503,7 +511,7 @@ void UCombineMeshesTool::UpdateExistingAsset()
 		// select the new actor
 		ToolSelectionUtil::SetNewActorSelection(GetToolManager(), SkipActor);
 	}
-#endif
+
 
 	
 	TArray<AActor*> Actors;
diff --git a/Engine/Plugins/Experimental/MeshModelingToolset/Source/MeshModelingTools/Private/CutMeshWithMeshTool.cpp b/Engine/Plugins/Experimental/MeshModelingToolset/Source/MeshModelingTools/Private/CutMeshWithMeshTool.cpp
--- a/Engine/Plugins/Experimental/MeshModelingToolset/Source/MeshModelingTools/Private/CutMeshWithMeshTool.cpp
+++ b/Engine/Plugins/Experimental/MeshModelingToolset/Source/MeshModelingTools/Private/CutMeshWithMeshTool.cpp
@@ -5,7 +5,7 @@
 #include "ToolSetupUtil.h"
 #include "BaseGizmos/TransformGizmoUtil.h"
 #include "Selection/ToolSelectionUtil.h"
-#include "AssetGenerationUtil.h"
+#include "ModelingObjectsCreationAPI.h"
 #include "DynamicMesh3.h"
 #include "MeshTransforms.h"
 #include "MeshDescriptionToDynamicMesh.h"
@@ -330,7 +330,7 @@ void UCutMeshWithMeshTool::Shutdown(EToolShutdownType ShutdownType)
 	HandleSourcesProperties->SaveProperties(this);
 	TransformProperties->SaveProperties(this);
 
-	FDynamicMeshOpResult Result = Preview->Shutdown();
+	FDynamicMeshOpResult OpResult = Preview->Shutdown();
 	// Restore (unhide) the source meshes
 	for ( int32 ci = 0; ci < Targets.Num(); ++ci)
 	{
@@ -350,14 +350,14 @@ void UCutMeshWithMeshTool::Shutdown(EToolShutdownType ShutdownType)
 		IPrimitiveComponentBackedTarget* UpdateTarget = TargetComponentInterface(0);
 		FTransform3d TargetToWorld = (FTransform3d)UpdateTarget->GetWorldTransform();
 		{
-			if (Result.Mesh->TriangleCount() > 0)
+			if (OpResult.Mesh->TriangleCount() > 0)
 			{
-				MeshTransforms::ApplyTransform(*Result.Mesh, Result.Transform);
-				MeshTransforms::ApplyTransformInverse(*Result.Mesh, TargetToWorld);
+				MeshTransforms::ApplyTransform(*OpResult.Mesh, OpResult.Transform);
+				MeshTransforms::ApplyTransformInverse(*OpResult.Mesh, TargetToWorld);
 				TargetMeshCommitterInterface(0)->CommitMeshDescription([&](const IMeshDescriptionCommitter::FCommitterParams& CommitParams)
 				{
 					FDynamicMeshToMeshDescription Converter;
-					Converter.Convert(Result.Mesh.Get(), *CommitParams.MeshDescriptionOut);
+					Converter.Convert(OpResult.Mesh.Get(), *CommitParams.MeshDescriptionOut);
 				});
 				TargetMaterialInterface(0)->CommitMaterialSetUpdate(MaterialSet, true);
 			}
@@ -367,19 +367,23 @@ void UCutMeshWithMeshTool::Shutdown(EToolShutdownType ShutdownType)
 		// create intersection asset
 		if ( IntersectionMesh.TriangleCount() > 0)
 		{
-			MeshTransforms::ApplyTransform(IntersectionMesh, Result.Transform);
+			MeshTransforms::ApplyTransform(IntersectionMesh, OpResult.Transform);
 			MeshTransforms::ApplyTransformInverse(IntersectionMesh, TargetToWorld);
 			FTransform3d NewTransform = TargetToWorld;
 
-			FString CurName = AssetGenerationUtil::GetComponentAssetBaseName(UpdateTarget->GetOwnerComponent());
+			FString CurName = UE::Modeling::GetComponentAssetBaseName(UpdateTarget->GetOwnerComponent());
 			FString UseBaseName = FString::Printf(TEXT("%s_%s"), *CurName, TEXT("CutPart") );
 
-			TArray<UMaterialInterface*> Materials = GetOutputMaterials();
-			AActor* NewActor = AssetGenerationUtil::GenerateStaticMeshActor( AssetAPI, TargetWorld,
-				&IntersectionMesh, NewTransform, UseBaseName, Materials);
-			if (NewActor != nullptr)
+			FCreateMeshObjectParams NewMeshObjectParams;
+			NewMeshObjectParams.TargetWorld = TargetWorld;
+			NewMeshObjectParams.Transform = (FTransform)NewTransform;
+			NewMeshObjectParams.BaseName = UseBaseName;
+			NewMeshObjectParams.Materials = GetOutputMaterials();
+			NewMeshObjectParams.SetMesh(&IntersectionMesh);
+			FCreateMeshObjectResult Result = UE::Modeling::CreateMeshObject(GetToolManager(), MoveTemp(NewMeshObjectParams));
+			if (Result.IsOK() && Result.NewActor != nullptr)
 			{
-				SelectActors.Add(NewActor);
+				SelectActors.Add(Result.NewActor);
 			}
 		}
 
@@ -387,7 +391,6 @@ void UCutMeshWithMeshTool::Shutdown(EToolShutdownType ShutdownType)
 		GetToolManager()->EndUndoTransaction();
 	}
 
-
 	UInteractiveGizmoManager* GizmoManager = GetToolManager()->GetPairedGizmoManager();
 	GizmoManager->DestroyAllGizmosByOwner(this);
 }
diff --git a/Engine/Plugins/Experimental/MeshModelingToolset/Source/MeshModelingTools/Private/DrawAndRevolveTool.cpp b/Engine/Plugins/Experimental/MeshModelingToolset/Source/MeshModelingTools/Private/DrawAndRevolveTool.cpp
--- a/Engine/Plugins/Experimental/MeshModelingToolset/Source/MeshModelingTools/Private/DrawAndRevolveTool.cpp
+++ b/Engine/Plugins/Experimental/MeshModelingToolset/Source/MeshModelingTools/Private/DrawAndRevolveTool.cpp
@@ -2,7 +2,6 @@
 
 #include "DrawAndRevolveTool.h"
 
-#include "AssetGenerationUtil.h"
 #include "BaseGizmos/GizmoRenderingUtil.h"
 #include "BaseBehaviors/SingleClickBehavior.h"
 #include "BaseBehaviors/MouseHoverBehavior.h"
@@ -12,6 +11,7 @@
 #include "Mechanics/ConstructionPlaneMechanic.h"
 #include "Mechanics/CurveControlPointsMechanic.h"
 #include "Properties/MeshMaterialProperties.h"
+#include "ModelingObjectsCreationAPI.h"
 #include "Selection/ToolSelectionUtil.h"
 #include "ToolSceneQueriesUtil.h"
 #include "ToolSetupUtil.h"
@@ -29,7 +29,7 @@ const FText EditModeMessage = LOCTEXT("CurveEditing", "Click points to select th
 
 bool UDrawAndRevolveToolBuilder::CanBuildTool(const FToolBuilderState& SceneState) const
 {
-	return (this->AssetAPI != nullptr);
+	return true;
 }
 
 UInteractiveTool* UDrawAndRevolveToolBuilder::BuildTool(const FToolBuilderState& SceneState) const
@@ -37,7 +37,6 @@ UInteractiveTool* UDrawAndRevolveToolBuilder::BuildTool(const FToolBuilderState&
 	UDrawAndRevolveTool* NewTool = NewObject<UDrawAndRevolveTool>(SceneState.ToolManager);
 
 	NewTool->SetWorld(SceneState.World);
-	NewTool->SetAssetAPI(AssetAPI);
 	return NewTool;
 }
 
@@ -229,21 +228,25 @@ void UDrawAndRevolveTool::Shutdown(EToolShutdownType ShutdownType)
 	}
 }
 
-void UDrawAndRevolveTool::GenerateAsset(const FDynamicMeshOpResult& Result)
+void UDrawAndRevolveTool::GenerateAsset(const FDynamicMeshOpResult& OpResult)
 {
-	if (Result.Mesh->TriangleCount() <= 0)
+	if (OpResult.Mesh->TriangleCount() <= 0)
 	{
 		return;
 	}
 
 	GetToolManager()->BeginUndoTransaction(LOCTEXT("RevolveToolTransactionName", "Revolve Tool"));
 
-	AActor* NewActor = AssetGenerationUtil::GenerateStaticMeshActor(
-		AssetAPI, TargetWorld, Result.Mesh.Get(), Result.Transform, TEXT("RevolveResult"), MaterialProperties->Material.Get());
-
-	if (NewActor != nullptr)
+	FCreateMeshObjectParams NewMeshObjectParams;
+	NewMeshObjectParams.TargetWorld = TargetWorld;
+	NewMeshObjectParams.Transform = (FTransform)OpResult.Transform;
+	NewMeshObjectParams.BaseName = TEXT("Revolve");
+	NewMeshObjectParams.Materials.Add(MaterialProperties->Material.Get());
+	NewMeshObjectParams.SetMesh(OpResult.Mesh.Get());
+	FCreateMeshObjectResult Result = UE::Modeling::CreateMeshObject(GetToolManager(), MoveTemp(NewMeshObjectParams));
+	if (Result.IsOK() && Result.NewActor != nullptr)
 	{
-		ToolSelectionUtil::SetNewActorSelection(GetToolManager(), NewActor);
+		ToolSelectionUtil::SetNewActorSelection(GetToolManager(), Result.NewActor);
 	}
 
 	GetToolManager()->EndUndoTransaction();
diff --git a/Engine/Plugins/Experimental/MeshModelingToolset/Source/MeshModelingTools/Private/DrawPolyPathTool.cpp b/Engine/Plugins/Experimental/MeshModelingToolset/Source/MeshModelingTools/Private/DrawPolyPathTool.cpp
--- a/Engine/Plugins/Experimental/MeshModelingToolset/Source/MeshModelingTools/Private/DrawPolyPathTool.cpp
+++ b/Engine/Plugins/Experimental/MeshModelingToolset/Source/MeshModelingTools/Private/DrawPolyPathTool.cpp
@@ -14,7 +14,7 @@
 #include "DynamicMesh/MeshIndexUtil.h"
 #include "Generators/RectangleMeshGenerator.h"
 #include "Distance/DistLine3Line3.h"
-#include "AssetGenerationUtil.h"
+#include "ModelingObjectsCreationAPI.h"
 #include "MeshTransforms.h"
 #include "Selection/ToolSelectionUtil.h"
 #include "Operations/ExtrudeMesh.h"
@@ -150,14 +150,13 @@ namespace
  */
 bool UDrawPolyPathToolBuilder::CanBuildTool(const FToolBuilderState& SceneState) const
 {
-	return (this->AssetAPI != nullptr);
+	return true;
 }
 
 UInteractiveTool* UDrawPolyPathToolBuilder::BuildTool(const FToolBuilderState& SceneState) const
 {
 	UDrawPolyPathTool* NewTool = NewObject<UDrawPolyPathTool>(SceneState.ToolManager);
 	NewTool->SetWorld(SceneState.World);
-	NewTool->SetAssetAPI(AssetAPI);
 	return NewTool;
 }
 
@@ -169,12 +168,6 @@ void UDrawPolyPathTool::SetWorld(UWorld* World)
 	this->TargetWorld = World;
 }
 
-void UDrawPolyPathTool::SetAssetAPI(IAssetGenerationAPI* AssetAPIIn)
-{
-	this->AssetAPI = AssetAPIIn;
-}
-
-
 void UDrawPolyPathTool::Setup()
 {
 	UInteractiveTool::Setup();
@@ -775,12 +768,16 @@ void UDrawPolyPathTool::EmitNewObject(EDrawPolyPathOutputMode OutputMode)
 
 	GetToolManager()->BeginUndoTransaction(LOCTEXT("CreatePolyPath", "Create PolyPath"));
 
-	AActor* NewActor = AssetGenerationUtil::GenerateStaticMeshActor(
-		AssetAPI, TargetWorld,
-		&PathMesh, MeshTransform.ToTransform(), TEXT("Path"), MaterialProperties->Material.Get());
-	if (NewActor != nullptr)
+	FCreateMeshObjectParams NewMeshObjectParams;
+	NewMeshObjectParams.TargetWorld = TargetWorld;
+	NewMeshObjectParams.Transform = MeshTransform.ToFTransform();
+	NewMeshObjectParams.BaseName = TEXT("Path");
+	NewMeshObjectParams.Materials.Add(MaterialProperties->Material.Get());
+	NewMeshObjectParams.SetMesh(&PathMesh);
+	FCreateMeshObjectResult Result = UE::Modeling::CreateMeshObject(GetToolManager(), MoveTemp(NewMeshObjectParams));
+	if (Result.IsOK() && Result.NewActor != nullptr)
 	{
-		ToolSelectionUtil::SetNewActorSelection(GetToolManager(), NewActor);
+		ToolSelectionUtil::SetNewActorSelection(GetToolManager(), Result.NewActor);
 	}
 
 	GetToolManager()->EndUndoTransaction();
diff --git a/Engine/Plugins/Experimental/MeshModelingToolset/Source/MeshModelingTools/Private/DrawPolygonTool.cpp b/Engine/Plugins/Experimental/MeshModelingToolset/Source/MeshModelingTools/Private/DrawPolygonTool.cpp
--- a/Engine/Plugins/Experimental/MeshModelingToolset/Source/MeshModelingTools/Private/DrawPolygonTool.cpp
+++ b/Engine/Plugins/Experimental/MeshModelingToolset/Source/MeshModelingTools/Private/DrawPolygonTool.cpp
@@ -32,7 +32,7 @@
 
 #include "Selection/SelectClickedAction.h"
 #include "Selection/ToolSelectionUtil.h"
-#include "AssetGenerationUtil.h"
+#include "ModelingObjectsCreationAPI.h"
 
 #include "ExplicitUseGeometryMathTypes.h"		// using UE::Geometry::(math types)
 using namespace UE::Geometry;
@@ -48,14 +48,13 @@ constexpr int CurrentGridSnapID = FPointPlanarSnapSolver::BaseExternalPointID +
 
 bool UDrawPolygonToolBuilder::CanBuildTool(const FToolBuilderState& SceneState) const
 {
-	return (this->AssetAPI != nullptr);
+	return true;
 }
 
 UInteractiveTool* UDrawPolygonToolBuilder::BuildTool(const FToolBuilderState& SceneState) const
 {
 	UDrawPolygonTool* NewTool = NewObject<UDrawPolygonTool>(SceneState.ToolManager);
 	NewTool->SetWorld(SceneState.World);
-	NewTool->SetAssetAPI(AssetAPI);
 	return NewTool;
 }
 
@@ -82,11 +81,6 @@ void UDrawPolygonTool::SetWorld(UWorld* World)
 	this->TargetWorld = World;
 }
 
-void UDrawPolygonTool::SetAssetAPI(IAssetGenerationAPI* AssetAPIIn)
-{
-	this->AssetAPI = AssetAPIIn;
-}
-
 void UDrawPolygonTool::Setup()
 {
 	UInteractiveTool::Setup();
@@ -1000,7 +994,6 @@ void UDrawPolygonTool::EmitCurrentPolygon()
 	FString BaseName = (PolygonProperties->OutputMode == EDrawPolygonOutputMode::MeshedPolygon) ?
 		TEXT("Polygon") : TEXT("Extrude");
 
-#if WITH_EDITOR
 	// generate new mesh
 	FFrame3d PlaneFrameOut;
 	FDynamicMesh3 Mesh;
@@ -1015,18 +1008,20 @@ void UDrawPolygonTool::EmitCurrentPolygon()
 
 	GetToolManager()->BeginUndoTransaction(LOCTEXT("CreatePolygon", "Create Polygon"));
 
-	AActor* NewActor = AssetGenerationUtil::GenerateStaticMeshActor(
-		AssetAPI, TargetWorld,
-		&Mesh, PlaneFrameOut.ToTransform(), BaseName, MaterialProperties->Material.Get());
-	if (NewActor != nullptr)
+	FCreateMeshObjectParams NewMeshObjectParams;
+	NewMeshObjectParams.TargetWorld = TargetWorld;
+	NewMeshObjectParams.Transform = PlaneFrameOut.ToFTransform();
+	NewMeshObjectParams.BaseName = BaseName;
+	NewMeshObjectParams.Materials.Add(MaterialProperties->Material.Get());
+	NewMeshObjectParams.SetMesh(&Mesh);
+	FCreateMeshObjectResult Result = UE::Modeling::CreateMeshObject(GetToolManager(), MoveTemp(NewMeshObjectParams));
+	if (Result.IsOK() && Result.NewActor != nullptr)
 	{
-		ToolSelectionUtil::SetNewActorSelection(GetToolManager(), NewActor);
+		ToolSelectionUtil::SetNewActorSelection(GetToolManager(), Result.NewActor);
 	}
 
 	GetToolManager()->EndUndoTransaction();
-#else
-	checkNoEntry();
-#endif
+
 	ResetPolygon();
 }
 
diff --git a/Engine/Plugins/Experimental/MeshModelingToolset/Source/MeshModelingTools/Private/EditNormalsTool.cpp b/Engine/Plugins/Experimental/MeshModelingToolset/Source/MeshModelingTools/Private/EditNormalsTool.cpp
--- a/Engine/Plugins/Experimental/MeshModelingToolset/Source/MeshModelingTools/Private/EditNormalsTool.cpp
+++ b/Engine/Plugins/Experimental/MeshModelingToolset/Source/MeshModelingTools/Private/EditNormalsTool.cpp
@@ -15,7 +15,6 @@
 
 #include "InteractiveGizmoManager.h"
 
-#include "AssetGenerationUtil.h"
 #include "AssetUtils/MeshDescriptionUtil.h"
 #include "Engine/Classes/Engine/StaticMesh.h"
 #include "Engine/Classes/Components/StaticMeshComponent.h"
@@ -57,7 +56,6 @@ UInteractiveTool* UEditNormalsToolBuilder::BuildTool(const FToolBuilderState& Sc
 	TArray<TObjectPtr<UToolTarget>> Targets = SceneState.TargetManager->BuildAllSelectedTargetable(SceneState, GetTargetRequirements());
 	NewTool->SetTargets(MoveTemp(Targets));
 	NewTool->SetWorld(SceneState.World);
-	NewTool->SetAssetAPI(AssetAPI);
 
 	return NewTool;
 }
@@ -194,11 +192,6 @@ void UEditNormalsTool::Shutdown(EToolShutdownType ShutdownType)
 	}
 }
 
-void UEditNormalsTool::SetAssetAPI(IAssetGenerationAPI* AssetAPIIn)
-{
-	this->AssetAPI = AssetAPIIn;
-}
-
 TUniquePtr<FDynamicMeshOperator> UEditNormalsOperatorFactory::MakeNewOperator()
 {
 	TUniquePtr<FEditNormalsOp> NormalsOp = MakeUnique<FEditNormalsOp>();
diff --git a/Engine/Plugins/Experimental/MeshModelingToolset/Source/MeshModelingTools/Private/MeshAttributePaintTool.cpp b/Engine/Plugins/Experimental/MeshModelingToolset/Source/MeshModelingTools/Private/MeshAttributePaintTool.cpp
--- a/Engine/Plugins/Experimental/MeshModelingToolset/Source/MeshModelingTools/Private/MeshAttributePaintTool.cpp
+++ b/Engine/Plugins/Experimental/MeshModelingToolset/Source/MeshModelingTools/Private/MeshAttributePaintTool.cpp
@@ -123,7 +123,6 @@ UMeshSurfacePointTool* UMeshAttributePaintToolBuilder::CreateNewTool(const FTool
 {
 	UMeshAttributePaintTool* SelectionTool = NewObject<UMeshAttributePaintTool>(SceneState.ToolManager);
 	SelectionTool->SetWorld(SceneState.World);
-	SelectionTool->SetAssetAPI(AssetAPI);
 
 	if (ColorMapFactory)
 	{
@@ -141,10 +140,6 @@ void UMeshAttributePaintTool::SetWorld(UWorld* World)
 	this->TargetWorld = World;
 }
 
-void UMeshAttributePaintTool::SetAssetAPI(IToolsContextAssetAPI* AssetAPIIn)
-{
-	this->AssetAPI = AssetAPIIn;
-}
 
 
 void UMeshAttributePaintTool::Setup()
diff --git a/Engine/Plugins/Experimental/MeshModelingToolset/Source/MeshModelingTools/Private/MeshSelectionTool.cpp b/Engine/Plugins/Experimental/MeshModelingToolset/Source/MeshModelingTools/Private/MeshSelectionTool.cpp
--- a/Engine/Plugins/Experimental/MeshModelingToolset/Source/MeshModelingTools/Private/MeshSelectionTool.cpp
+++ b/Engine/Plugins/Experimental/MeshModelingToolset/Source/MeshModelingTools/Private/MeshSelectionTool.cpp
@@ -12,7 +12,7 @@
 #include "Selections/MeshConnectedComponents.h"
 #include "MeshRegionBoundaryLoops.h"
 #include "DynamicMesh/MeshIndexUtil.h"
-#include "AssetGenerationUtil.h"
+#include "ModelingObjectsCreationAPI.h"
 #include "ToolSetupUtil.h"
 #include "Selections/MeshConnectedComponents.h"
 #include "Selections/MeshFaceSelection.h"
@@ -35,7 +35,6 @@ UMeshSurfacePointTool* UMeshSelectionToolBuilder::CreateNewTool(const FToolBuild
 {
 	UMeshSelectionTool* SelectionTool = NewObject<UMeshSelectionTool>(SceneState.ToolManager);
 	SelectionTool->SetWorld(SceneState.World);
-	SelectionTool->SetAssetAPI(AssetAPI);
 	return SelectionTool;
 }
 
@@ -62,13 +61,6 @@ void UMeshSelectionTool::SetWorld(UWorld* World)
 	this->TargetWorld = World;
 }
 
-void UMeshSelectionTool::SetAssetAPI(IAssetGenerationAPI* AssetAPIIn)
-{
-	this->AssetAPI = AssetAPIIn;
-}
-
-
-
 void UMeshSelectionTool::Setup()
 {
 	UDynamicMeshBrushTool::Setup();
@@ -1394,12 +1386,20 @@ void UMeshSelectionTool::SeparateSelectedTriangles()
 	{
 		Materials.Add(TargetMaterial->GetMaterial(MaterialIdx));
 	}
-	AActor* NewActor = AssetGenerationUtil::GenerateStaticMeshActor(
-		AssetAPI, TargetWorld, &SeparatedMesh, Transform, TEXT("Submesh"), Materials);
-	SpawnedActors.Add(NewActor);
-	GetToolManager()->EndUndoTransaction();
 
-	// todo: undo won't remove this asset...
+	FCreateMeshObjectParams NewMeshObjectParams;
+	NewMeshObjectParams.TargetWorld = TargetWorld;
+	NewMeshObjectParams.Transform = (FTransform)Transform;
+	NewMeshObjectParams.BaseName = TEXT("Submesh");
+	NewMeshObjectParams.Materials = Materials;
+	NewMeshObjectParams.SetMesh(&SeparatedMesh);
+	FCreateMeshObjectResult Result = UE::Modeling::CreateMeshObject(GetToolManager(), MoveTemp(NewMeshObjectParams));
+	if (Result.IsOK() && Result.NewActor != nullptr)
+	{
+		SpawnedActors.Add(Result.NewActor);
+	}
+
+	GetToolManager()->EndUndoTransaction();
 
 	// delete selected triangles from this mesh
 	DeleteSelectedTriangles();
diff --git a/Engine/Plugins/Experimental/MeshModelingToolset/Source/MeshModelingTools/Private/MirrorTool.cpp b/Engine/Plugins/Experimental/MeshModelingToolset/Source/MeshModelingTools/Private/MirrorTool.cpp
--- a/Engine/Plugins/Experimental/MeshModelingToolset/Source/MeshModelingTools/Private/MirrorTool.cpp
+++ b/Engine/Plugins/Experimental/MeshModelingToolset/Source/MeshModelingTools/Private/MirrorTool.cpp
@@ -2,7 +2,7 @@
 
 #include "MirrorTool.h"
 
-#include "AssetGenerationUtil.h"
+#include "ModelingObjectsCreationAPI.h"
 #include "BaseBehaviors/KeyAsModifierInputBehavior.h"
 #include "BaseBehaviors/SingleClickBehavior.h"
 #include "CompositionOps/MirrorOp.h"
@@ -43,7 +43,7 @@ const FToolTargetTypeRequirements& UMirrorToolBuilder::GetTargetRequirements() c
 
 bool UMirrorToolBuilder::CanBuildTool(const FToolBuilderState& SceneState) const
 {
-	return AssetAPI != nullptr && SceneState.TargetManager->CountSelectedAndTargetable(SceneState, GetTargetRequirements()) > 0;
+	return SceneState.TargetManager->CountSelectedAndTargetable(SceneState, GetTargetRequirements()) > 0;
 }
 
 UInteractiveTool* UMirrorToolBuilder::BuildTool(const FToolBuilderState& SceneState) const
@@ -53,7 +53,6 @@ UInteractiveTool* UMirrorToolBuilder::BuildTool(const FToolBuilderState& SceneSt
 	TArray<TObjectPtr<UToolTarget>> Targets = SceneState.TargetManager->BuildAllSelectedTargetable(SceneState, GetTargetRequirements());
 	NewTool->SetTargets(MoveTemp(Targets));
 	NewTool->SetWorld(SceneState.World);
-	NewTool->SetAssetAPI(AssetAPI);
 
 	return NewTool;
 }
@@ -132,11 +131,6 @@ void UMirrorTool::SetWorld(UWorld* World)
 	TargetWorld = World;
 }
 
-void UMirrorTool::SetAssetAPI(IAssetGenerationAPI* NewAssetApi)
-{
-	AssetAPI = NewAssetApi;
-}
-
 void UMirrorTool::OnPropertyModified(UObject* PropertySet, FProperty* Property)
 {
 	// Editing the "show preview" option changes whether we need to be displaying the preview or the original mesh.
@@ -480,12 +474,16 @@ void UMirrorTool::GenerateAsset(const TArray<FDynamicMeshOpResult>& Results)
 				Materials.Add(TargetMaterial->GetMaterial(MaterialIdx));
 			}
 
-			// Create the new actor
-			AActor* NewActor = AssetGenerationUtil::GenerateStaticMeshActor(
-				AssetAPI, TargetWorld, Mesh, Results[OrigMeshIdx].Transform, TEXT("MirrorResult"), Materials);
-			if (NewActor != nullptr)
+			FCreateMeshObjectParams NewMeshObjectParams;
+			NewMeshObjectParams.TargetWorld = TargetWorld;
+			NewMeshObjectParams.Transform = (FTransform)Results[OrigMeshIdx].Transform;
+			NewMeshObjectParams.BaseName = TEXT("Mirror");
+			NewMeshObjectParams.Materials = Materials;
+			NewMeshObjectParams.SetMesh(Mesh);
+			FCreateMeshObjectResult Result = UE::Modeling::CreateMeshObject(GetToolManager(), MoveTemp(NewMeshObjectParams));
+			if (Result.IsOK() && Result.NewActor != nullptr)
 			{
-				NewSelection.Actors.Add(NewActor);
+				NewSelection.Actors.Add(Result.NewActor);
 			}
 
 			// Remove the original actor
diff --git a/Engine/Plugins/Experimental/MeshModelingToolset/Source/MeshModelingTools/Private/Physics/ExtractCollisionGeometryTool.cpp b/Engine/Plugins/Experimental/MeshModelingToolset/Source/MeshModelingTools/Private/Physics/ExtractCollisionGeometryTool.cpp
--- a/Engine/Plugins/Experimental/MeshModelingToolset/Source/MeshModelingTools/Private/Physics/ExtractCollisionGeometryTool.cpp
+++ b/Engine/Plugins/Experimental/MeshModelingToolset/Source/MeshModelingTools/Private/Physics/ExtractCollisionGeometryTool.cpp
@@ -4,7 +4,7 @@
 #include "InteractiveToolManager.h"
 #include "ToolBuilderUtil.h"
 #include "ToolSetupUtil.h"
-#include "AssetGenerationUtil.h"
+#include "ModelingObjectsCreationAPI.h"
 #include "Selection/ToolSelectionUtil.h"
 #include "Drawing/PreviewGeometryActor.h"
 #include "Util/ColorConstants.h"
@@ -130,12 +130,16 @@ void UExtractCollisionGeometryTool::Shutdown(EToolShutdownType ShutdownType)
 
 		GetToolManager()->BeginUndoTransaction(LOCTEXT("CreateCollisionMesh", "Collision To Mesh"));
 
-		AActor* NewActor = AssetGenerationUtil::GenerateStaticMeshActor(
-			AssetAPI, TargetWorld,
-			&CurrentMesh, Transform, NewName, UseMaterial);
-		if (NewActor != nullptr)
+		FCreateMeshObjectParams NewMeshObjectParams;
+		NewMeshObjectParams.TargetWorld = TargetWorld;
+		NewMeshObjectParams.Transform = (FTransform)Transform;
+		NewMeshObjectParams.BaseName = NewName;
+		NewMeshObjectParams.Materials.Add(UseMaterial);
+		NewMeshObjectParams.SetMesh(&CurrentMesh);
+		FCreateMeshObjectResult Result = UE::Modeling::CreateMeshObject(GetToolManager(), MoveTemp(NewMeshObjectParams));
+		if (Result.IsOK() && Result.NewActor != nullptr)
 		{
-			ToolSelectionUtil::SetNewActorSelection(GetToolManager(), NewActor);
+			ToolSelectionUtil::SetNewActorSelection(GetToolManager(), Result.NewActor);
 		}
 
 		GetToolManager()->EndUndoTransaction();
diff --git a/Engine/Plugins/Experimental/MeshModelingToolset/Source/MeshModelingTools/Private/PlaneCutTool.cpp b/Engine/Plugins/Experimental/MeshModelingToolset/Source/MeshModelingTools/Private/PlaneCutTool.cpp
--- a/Engine/Plugins/Experimental/MeshModelingToolset/Source/MeshModelingTools/Private/PlaneCutTool.cpp
+++ b/Engine/Plugins/Experimental/MeshModelingToolset/Source/MeshModelingTools/Private/PlaneCutTool.cpp
@@ -22,7 +22,7 @@
 #include "BaseGizmos/TransformGizmo.h"
 
 #include "Drawing/MeshDebugDrawing.h"
-#include "AssetGenerationUtil.h"
+#include "ModelingObjectsCreationAPI.h"
 
 #include "Changes/ToolCommandChangeSequence.h"
 
@@ -60,7 +60,7 @@ const FToolTargetTypeRequirements& UPlaneCutToolBuilder::GetTargetRequirements()
 
 bool UPlaneCutToolBuilder::CanBuildTool(const FToolBuilderState& SceneState) const
 {
-	return AssetAPI != nullptr && SceneState.TargetManager->CountSelectedAndTargetable(SceneState, GetTargetRequirements()) > 0;
+	return SceneState.TargetManager->CountSelectedAndTargetable(SceneState, GetTargetRequirements()) > 0;
 }
 
 UInteractiveTool* UPlaneCutToolBuilder::BuildTool(const FToolBuilderState& SceneState) const
@@ -70,7 +70,6 @@ UInteractiveTool* UPlaneCutToolBuilder::BuildTool(const FToolBuilderState& Scene
 	TArray<TObjectPtr<UToolTarget>> Targets = SceneState.TargetManager->BuildAllSelectedTargetable(SceneState, GetTargetRequirements());
 	NewTool->SetTargets(MoveTemp(Targets));
 	NewTool->SetWorld(SceneState.World);
-	NewTool->SetAssetAPI(AssetAPI);
 
 	return NewTool;
 }
@@ -302,11 +301,6 @@ void UPlaneCutTool::Shutdown(EToolShutdownType ShutdownType)
 	}
 }
 
-void UPlaneCutTool::SetAssetAPI(IAssetGenerationAPI* AssetAPIIn)
-{
-	this->AssetAPI = AssetAPIIn;
-}
-
 TUniquePtr<FDynamicMeshOperator> UPlaneCutOperatorFactory::MakeNewOperator()
 {
 	TUniquePtr<FPlaneCutOp> CutOp = MakeUnique<FPlaneCutOp>();
@@ -552,12 +546,17 @@ void UPlaneCutTool::GenerateAsset(const TArray<FDynamicMeshOpResult>& Results)
 			// add all the additional meshes
 			for (int AddMeshIdx = 1; AddMeshIdx < SplitMeshes.Num(); AddMeshIdx++)
 			{
-				AActor* NewActor = AssetGenerationUtil::GenerateStaticMeshActor(
-					AssetAPI, TargetWorld,
-					&SplitMeshes[AddMeshIdx], Results[OrigMeshIdx].Transform, TEXT("PlaneCutOtherPart"), ComponentMaterialSet.Materials, AssetMaterialSet.Materials);
-				if (NewActor != nullptr)
+				FCreateMeshObjectParams NewMeshObjectParams;
+				NewMeshObjectParams.TargetWorld = TargetWorld;
+				NewMeshObjectParams.Transform = (FTransform)Results[OrigMeshIdx].Transform;
+				NewMeshObjectParams.BaseName = TEXT("PlaneCutOtherPart");
+				NewMeshObjectParams.Materials = ComponentMaterialSet.Materials;
+				NewMeshObjectParams.AssetMaterials = AssetMaterialSet.Materials;
+				NewMeshObjectParams.SetMesh(&SplitMeshes[AddMeshIdx]);
+				FCreateMeshObjectResult Result = UE::Modeling::CreateMeshObject(GetToolManager(), MoveTemp(NewMeshObjectParams));
+				if (Result.IsOK() && Result.NewActor != nullptr)
 				{
-					NewSelection.Actors.Add(NewActor);
+					NewSelection.Actors.Add(Result.NewActor);
 				}
 			}
 		}
diff --git a/Engine/Plugins/Experimental/MeshModelingToolset/Source/MeshModelingTools/Private/ProjectToTargetTool.cpp b/Engine/Plugins/Experimental/MeshModelingToolset/Source/MeshModelingTools/Private/ProjectToTargetTool.cpp
--- a/Engine/Plugins/Experimental/MeshModelingToolset/Source/MeshModelingTools/Private/ProjectToTargetTool.cpp
+++ b/Engine/Plugins/Experimental/MeshModelingToolset/Source/MeshModelingTools/Private/ProjectToTargetTool.cpp
@@ -38,7 +38,6 @@ UInteractiveTool* UProjectToTargetToolBuilder::BuildTool(const FToolBuilderState
 	TArray<TObjectPtr<UToolTarget>> Targets = SceneState.TargetManager->BuildAllSelectedTargetable(SceneState, GetTargetRequirements());
 	NewTool->SetTargets(MoveTemp(Targets));
 	NewTool->SetWorld(SceneState.World);
-	NewTool->SetAssetAPI(AssetAPI);
 
 	return NewTool;
 }
diff --git a/Engine/Plugins/Experimental/MeshModelingToolset/Source/MeshModelingTools/Private/RemeshMeshTool.cpp b/Engine/Plugins/Experimental/MeshModelingToolset/Source/MeshModelingTools/Private/RemeshMeshTool.cpp
--- a/Engine/Plugins/Experimental/MeshModelingToolset/Source/MeshModelingTools/Private/RemeshMeshTool.cpp
+++ b/Engine/Plugins/Experimental/MeshModelingToolset/Source/MeshModelingTools/Private/RemeshMeshTool.cpp
@@ -13,8 +13,6 @@
 #include "MeshDescriptionToDynamicMesh.h"
 #include "DynamicMeshToMeshDescription.h"
 
-#include "AssetGenerationUtil.h"
-
 #include "SceneManagement.h" // for FPrimitiveDrawInterface
 
 #include "TargetInterfaces/MaterialProvider.h"
@@ -54,7 +52,6 @@ UInteractiveTool* URemeshMeshToolBuilder::BuildTool(const FToolBuilderState& Sce
 	TArray<TObjectPtr<UToolTarget>> Targets = SceneState.TargetManager->BuildAllSelectedTargetable(SceneState, GetTargetRequirements());
 	NewTool->SetTargets(MoveTemp(Targets));
 	NewTool->SetWorld(SceneState.World);
-	NewTool->SetAssetAPI(AssetAPI);
 
 	return NewTool;
 }
@@ -97,11 +94,6 @@ void URemeshMeshTool::SetWorld(UWorld* World)
 	this->TargetWorld = World;
 }
 
-void URemeshMeshTool::SetAssetAPI(IToolsContextAssetAPI* AssetAPIIn)
-{
-	this->AssetAPI = AssetAPIIn;
-}
-
 void URemeshMeshTool::Setup()
 {
 	UInteractiveTool::Setup();
diff --git a/Engine/Plugins/Experimental/MeshModelingToolset/Source/MeshModelingTools/Private/RemoveOccludedTrianglesTool.cpp b/Engine/Plugins/Experimental/MeshModelingToolset/Source/MeshModelingTools/Private/RemoveOccludedTrianglesTool.cpp
--- a/Engine/Plugins/Experimental/MeshModelingToolset/Source/MeshModelingTools/Private/RemoveOccludedTrianglesTool.cpp
+++ b/Engine/Plugins/Experimental/MeshModelingToolset/Source/MeshModelingTools/Private/RemoveOccludedTrianglesTool.cpp
@@ -16,8 +16,6 @@
 
 #include "InteractiveGizmoManager.h"
 
-#include "AssetGenerationUtil.h"
-
 #include "Misc/MessageDialog.h"
 
 
@@ -66,7 +64,6 @@ UInteractiveTool* URemoveOccludedTrianglesToolBuilder::BuildTool(const FToolBuil
 	TArray<TObjectPtr<UToolTarget>> Targets = SceneState.TargetManager->BuildAllSelectedTargetable(SceneState, GetTargetRequirements());
 	NewTool->SetTargets(MoveTemp(Targets));
 	NewTool->SetWorld(SceneState.World);
-	NewTool->SetAssetAPI(AssetAPI);
 
 	return NewTool;
 }
@@ -357,11 +354,6 @@ void URemoveOccludedTrianglesTool::Shutdown(EToolShutdownType ShutdownType)
 	}
 }
 
-void URemoveOccludedTrianglesTool::SetAssetAPI(IAssetGenerationAPI* AssetAPIIn)
-{
-	this->AssetAPI = AssetAPIIn;
-}
-
 TUniquePtr<FDynamicMeshOperator> URemoveOccludedTrianglesOperatorFactory::MakeNewOperator()
 {
 	TUniquePtr<FRemoveOccludedTrianglesOp> Op = MakeUnique<FRemoveOccludedTrianglesOp>();
diff --git a/Engine/Plugins/Experimental/MeshModelingToolset/Source/MeshModelingTools/Private/RevolveBoundaryTool.cpp b/Engine/Plugins/Experimental/MeshModelingToolset/Source/MeshModelingTools/Private/RevolveBoundaryTool.cpp
--- a/Engine/Plugins/Experimental/MeshModelingToolset/Source/MeshModelingTools/Private/RevolveBoundaryTool.cpp
+++ b/Engine/Plugins/Experimental/MeshModelingToolset/Source/MeshModelingTools/Private/RevolveBoundaryTool.cpp
@@ -2,7 +2,7 @@
 
 #include "RevolveBoundaryTool.h"
 
-#include "AssetGenerationUtil.h"
+#include "ModelingObjectsCreationAPI.h"
 #include "BaseBehaviors/SingleClickBehavior.h"
 #include "CoreMinimal.h"
 #include "CompositionOps/CurveSweepOp.h"
@@ -272,21 +272,25 @@ void URevolveBoundaryTool::Shutdown(EToolShutdownType ShutdownType)
 	}
 }
 
-void URevolveBoundaryTool::GenerateAsset(const FDynamicMeshOpResult& Result)
+void URevolveBoundaryTool::GenerateAsset(const FDynamicMeshOpResult& OpResult)
 {
-	if (Result.Mesh->TriangleCount() <= 0)
+	if (OpResult.Mesh->TriangleCount() <= 0)
 	{
 		return;
 	}
 
 	GetToolManager()->BeginUndoTransaction(LOCTEXT("RevolveBoundaryToolTransactionName", "Revolve Tool"));
 
-	AActor* NewActor = AssetGenerationUtil::GenerateStaticMeshActor(
-		AssetAPI, TargetWorld, Result.Mesh.Get(), Result.Transform, TEXT("RevolveBoundaryResult"), MaterialProperties->Material.Get());
-
-	if (NewActor != nullptr)
+	FCreateMeshObjectParams NewMeshObjectParams;
+	NewMeshObjectParams.TargetWorld = TargetWorld;
+	NewMeshObjectParams.Transform = (FTransform)OpResult.Transform;
+	NewMeshObjectParams.BaseName = TEXT("Revolve");
+	NewMeshObjectParams.Materials.Add(MaterialProperties->Material.Get());
+	NewMeshObjectParams.SetMesh(OpResult.Mesh.Get());
+	FCreateMeshObjectResult Result = UE::Modeling::CreateMeshObject(GetToolManager(), MoveTemp(NewMeshObjectParams));
+	if (Result.IsOK() && Result.NewActor != nullptr)
 	{
-		ToolSelectionUtil::SetNewActorSelection(GetToolManager(), NewActor);
+		ToolSelectionUtil::SetNewActorSelection(GetToolManager(), Result.NewActor);
 	}
 
 	GetToolManager()->EndUndoTransaction();
diff --git a/Engine/Plugins/Experimental/MeshModelingToolset/Source/MeshModelingTools/Private/SelfUnionMeshesTool.cpp b/Engine/Plugins/Experimental/MeshModelingToolset/Source/MeshModelingTools/Private/SelfUnionMeshesTool.cpp
--- a/Engine/Plugins/Experimental/MeshModelingToolset/Source/MeshModelingTools/Private/SelfUnionMeshesTool.cpp
+++ b/Engine/Plugins/Experimental/MeshModelingToolset/Source/MeshModelingTools/Private/SelfUnionMeshesTool.cpp
@@ -13,8 +13,6 @@
 #include "MeshDescriptionToDynamicMesh.h"
 #include "DynamicMeshToMeshDescription.h"
 
-#include "AssetGenerationUtil.h"
-
 #include "ExplicitUseGeometryMathTypes.h"		// using UE::Geometry::(math types)
 using namespace UE::Geometry;
 
diff --git a/Engine/Plugins/Experimental/MeshModelingToolset/Source/MeshModelingTools/Private/UVProjectionTool.cpp b/Engine/Plugins/Experimental/MeshModelingToolset/Source/MeshModelingTools/Private/UVProjectionTool.cpp
--- a/Engine/Plugins/Experimental/MeshModelingToolset/Source/MeshModelingTools/Private/UVProjectionTool.cpp
+++ b/Engine/Plugins/Experimental/MeshModelingToolset/Source/MeshModelingTools/Private/UVProjectionTool.cpp
@@ -2,7 +2,6 @@
 
 #include "UVProjectionTool.h"
 #include "InteractiveToolManager.h"
-#include "InteractiveGizmoManager.h"
 #include "ToolTargetManager.h"
 #include "ToolBuilderUtil.h"
 #include "ToolSetupUtil.h"
diff --git a/Engine/Plugins/Experimental/MeshModelingToolset/Source/MeshModelingTools/Private/VolumeToMeshTool.cpp b/Engine/Plugins/Experimental/MeshModelingToolset/Source/MeshModelingTools/Private/VolumeToMeshTool.cpp
--- a/Engine/Plugins/Experimental/MeshModelingToolset/Source/MeshModelingTools/Private/VolumeToMeshTool.cpp
+++ b/Engine/Plugins/Experimental/MeshModelingToolset/Source/MeshModelingTools/Private/VolumeToMeshTool.cpp
@@ -11,7 +11,7 @@
 #include "Util/ColorConstants.h"
 #include "ToolSetupUtil.h"
 #include "Selection/ToolSelectionUtil.h"
-#include "AssetGenerationUtil.h"
+#include "ModelingObjectsCreationAPI.h"
 
 #include "DynamicMesh/Operations/MergeCoincidentMeshEdges.h"
 #include "CompGeom/PolygonTriangulation.h"
@@ -39,8 +39,6 @@ UInteractiveTool* UVolumeToMeshToolBuilder::BuildTool(const FToolBuilderState& S
 	UVolumeToMeshTool* NewTool = NewObject<UVolumeToMeshTool>(SceneState.ToolManager);
 
 	NewTool->SetWorld(SceneState.World);
-	check(AssetAPI);
-	NewTool->SetAssetAPI(AssetAPI);
 
 	AVolume* Volume = ToolBuilderUtil::FindFirstActorOfType<AVolume>(SceneState);
 	check(Volume != nullptr);
@@ -126,12 +124,16 @@ void UVolumeToMeshTool::Shutdown(EToolShutdownType ShutdownType)
 
 		GetToolManager()->BeginUndoTransaction(LOCTEXT("CreateMeshVolume", "Volume To Mesh"));
 
-		AActor* NewActor = AssetGenerationUtil::GenerateStaticMeshActor(
-			AssetAPI, TargetWorld,
-			&CurrentMesh, Transform, NewName, UseMaterial);
-		if (NewActor != nullptr)
+		FCreateMeshObjectParams NewMeshObjectParams;
+		NewMeshObjectParams.TargetWorld = TargetWorld;
+		NewMeshObjectParams.Transform = (FTransform)Transform;
+		NewMeshObjectParams.BaseName = NewName;
+		NewMeshObjectParams.Materials.Add(UseMaterial);
+		NewMeshObjectParams.SetMesh(&CurrentMesh);
+		FCreateMeshObjectResult Result = UE::Modeling::CreateMeshObject(GetToolManager(), MoveTemp(NewMeshObjectParams));
+		if (Result.IsOK() && Result.NewActor != nullptr)
 		{
-			ToolSelectionUtil::SetNewActorSelection(GetToolManager(), NewActor);
+			ToolSelectionUtil::SetNewActorSelection(GetToolManager(), Result.NewActor);
 		}
 
 		GetToolManager()->EndUndoTransaction();
diff --git a/Engine/Plugins/Experimental/MeshModelingToolset/Source/MeshModelingTools/Private/VoxelBlendMeshesTool.cpp b/Engine/Plugins/Experimental/MeshModelingToolset/Source/MeshModelingTools/Private/VoxelBlendMeshesTool.cpp
--- a/Engine/Plugins/Experimental/MeshModelingToolset/Source/MeshModelingTools/Private/VoxelBlendMeshesTool.cpp
+++ b/Engine/Plugins/Experimental/MeshModelingToolset/Source/MeshModelingTools/Private/VoxelBlendMeshesTool.cpp
@@ -20,8 +20,6 @@
 #include "BaseGizmos/GizmoComponents.h"
 #include "BaseGizmos/TransformGizmo.h"
 
-#include "AssetGenerationUtil.h"
-
 #include "CompositionOps/VoxelBlendMeshesOp.h"
 
 #include "ExplicitUseGeometryMathTypes.h"		// using UE::Geometry::(math types)
diff --git a/Engine/Plugins/Experimental/MeshModelingToolset/Source/MeshModelingTools/Private/VoxelMorphologyMeshesTool.cpp b/Engine/Plugins/Experimental/MeshModelingToolset/Source/MeshModelingTools/Private/VoxelMorphologyMeshesTool.cpp
--- a/Engine/Plugins/Experimental/MeshModelingToolset/Source/MeshModelingTools/Private/VoxelMorphologyMeshesTool.cpp
+++ b/Engine/Plugins/Experimental/MeshModelingToolset/Source/MeshModelingTools/Private/VoxelMorphologyMeshesTool.cpp
@@ -20,8 +20,6 @@
 #include "BaseGizmos/GizmoComponents.h"
 #include "BaseGizmos/TransformGizmo.h"
 
-#include "AssetGenerationUtil.h"
-
 #include "CompositionOps/VoxelMorphologyMeshesOp.h"
 
 #include "ExplicitUseGeometryMathTypes.h"		// using UE::Geometry::(math types)
diff --git a/Engine/Plugins/Experimental/MeshModelingToolset/Source/MeshModelingTools/Public/AddPatchTool.h b/Engine/Plugins/Experimental/MeshModelingToolset/Source/MeshModelingTools/Public/AddPatchTool.h
--- a/Engine/Plugins/Experimental/MeshModelingToolset/Source/MeshModelingTools/Public/AddPatchTool.h
+++ b/Engine/Plugins/Experimental/MeshModelingToolset/Source/MeshModelingTools/Public/AddPatchTool.h
@@ -11,8 +11,6 @@
 #include "Properties/MeshMaterialProperties.h"
 #include "AddPatchTool.generated.h"
 
-class IAssetGenerationAPI;
-
 PREDECLARE_USE_GEOMETRY_CLASS(FDynamicMesh3);
 
 /**
@@ -22,15 +20,7 @@ UCLASS()
 class MESHMODELINGTOOLS_API UAddPatchToolBuilder : public UInteractiveToolBuilder
 {
 	GENERATED_BODY()
-
 public:
-	IAssetGenerationAPI* AssetAPI;
-
-	UAddPatchToolBuilder() 
-	{
-		AssetAPI = nullptr;
-	}
-
 	virtual bool CanBuildTool(const FToolBuilderState& SceneState) const override;
 	virtual UInteractiveTool* BuildTool(const FToolBuilderState& SceneState) const override;
 };
@@ -78,7 +68,6 @@ class MESHMODELINGTOOLS_API UAddPatchTool : public USingleClickTool, public IHov
 
 public:
 	virtual void SetWorld(UWorld* World);
-	virtual void SetAssetAPI(IAssetGenerationAPI* AssetAPI);
 
 	virtual void Setup() override;
 	virtual void Shutdown(EToolShutdownType ShutdownType) override;
@@ -116,7 +105,6 @@ protected:
 
 protected:
 	UWorld* TargetWorld;
-	IAssetGenerationAPI* AssetAPI;
 
 	FBox WorldBounds;
 
diff --git a/Engine/Plugins/Experimental/MeshModelingToolset/Source/MeshModelingTools/Public/AddPrimitiveTool.h b/Engine/Plugins/Experimental/MeshModelingToolset/Source/MeshModelingTools/Public/AddPrimitiveTool.h
--- a/Engine/Plugins/Experimental/MeshModelingToolset/Source/MeshModelingTools/Public/AddPrimitiveTool.h
+++ b/Engine/Plugins/Experimental/MeshModelingToolset/Source/MeshModelingTools/Public/AddPrimitiveTool.h
@@ -14,7 +14,6 @@
 #include "AddPrimitiveTool.generated.h"
 
 PREDECLARE_USE_GEOMETRY_CLASS(FDynamicMesh3);
-class IAssetGenerationAPI;
 
 /**
  * Builder
@@ -38,12 +37,6 @@ public:
 		Stairs
 	};
 
-	IAssetGenerationAPI* AssetAPI;
-
-	UAddPrimitiveToolBuilder()
-	{
-		AssetAPI = nullptr;
-	}
 	virtual bool CanBuildTool(const FToolBuilderState& SceneState) const override;
 	virtual UInteractiveTool* BuildTool(const FToolBuilderState& SceneState) const override;
 
@@ -460,7 +453,6 @@ public:
 	UAddPrimitiveTool(const FObjectInitializer&);
 
 	virtual void SetWorld(UWorld* World);
-	virtual void SetAssetAPI(IAssetGenerationAPI* AssetAPI);
 
 	virtual void Setup() override;
 	virtual void Shutdown(EToolShutdownType ShutdownType) override;
@@ -520,7 +512,6 @@ protected:
 
 protected:
 	UWorld* TargetWorld;
-	IAssetGenerationAPI* AssetAPI;
 
 	void UpdatePreviewPosition(const FInputDeviceRay& ClickPos);
 	UE::Geometry::FFrame3f ShapeFrame;
diff --git a/Engine/Plugins/Experimental/MeshModelingToolset/Source/MeshModelingTools/Public/BakeMeshAttributeMapsTool.h b/Engine/Plugins/Experimental/MeshModelingToolset/Source/MeshModelingTools/Public/BakeMeshAttributeMapsTool.h
--- a/Engine/Plugins/Experimental/MeshModelingToolset/Source/MeshModelingTools/Public/BakeMeshAttributeMapsTool.h
+++ b/Engine/Plugins/Experimental/MeshModelingToolset/Source/MeshModelingTools/Public/BakeMeshAttributeMapsTool.h
@@ -21,7 +21,6 @@ struct FMeshDescription;
 class USimpleDynamicMeshComponent;
 class UMaterialInstanceDynamic;
 class UTexture2D;
-class IAssetGenerationAPI;
 PREDECLARE_GEOMETRY(template<typename RealType> class TMeshTangents);
 PREDECLARE_GEOMETRY(class FMeshImageBakingCache);
 using UE::Geometry::FImageDimensions;
@@ -43,8 +42,6 @@ class MESHMODELINGTOOLS_API UBakeMeshAttributeMapsToolBuilder : public UInteract
 	GENERATED_BODY()
 
 public:
-	IAssetGenerationAPI* AssetAPI = nullptr;
-
 	virtual bool CanBuildTool(const FToolBuilderState& SceneState) const override;
 	virtual UInteractiveTool* BuildTool(const FToolBuilderState& SceneState) const override;
 
@@ -355,8 +352,6 @@ class MESHMODELINGTOOLS_API UBakeMeshAttributeMapsTool : public UMultiSelectionT
 public:
 	UBakeMeshAttributeMapsTool() = default;
 
-	virtual void SetAssetAPI(IAssetGenerationAPI* AssetAPI);
-
 	virtual void Setup() override;
 	virtual void Shutdown(EToolShutdownType ShutdownType) override;
 
@@ -408,9 +403,6 @@ protected:
 
 	friend class FMeshMapBakerOp;
 
-
-	IAssetGenerationAPI* AssetAPI = nullptr;
-
 	USimpleDynamicMeshComponent* DynamicMeshComponent;
 
 	UPROPERTY()
diff --git a/Engine/Plugins/Experimental/MeshModelingToolset/Source/MeshModelingTools/Public/BakeTransformTool.h b/Engine/Plugins/Experimental/MeshModelingToolset/Source/MeshModelingTools/Public/BakeTransformTool.h
--- a/Engine/Plugins/Experimental/MeshModelingToolset/Source/MeshModelingTools/Public/BakeTransformTool.h
+++ b/Engine/Plugins/Experimental/MeshModelingToolset/Source/MeshModelingTools/Public/BakeTransformTool.h
@@ -17,9 +17,6 @@ struct FMeshDescription;
 class USimpleDynamicMeshComponent;
 
 
-
-
-
 /**
  *
  */
@@ -89,7 +86,6 @@ public:
 	virtual void Shutdown(EToolShutdownType ShutdownType) override;
 
 	virtual void SetWorld(UWorld* World);
-	virtual void SetAssetAPI(IToolsContextAssetAPI* AssetAPI);
 
 	virtual bool HasCancel() const override { return true; }
 	virtual bool HasAccept() const override { return true; }
@@ -101,7 +97,6 @@ protected:
 
 protected:
 	UWorld* TargetWorld;
-	IToolsContextAssetAPI* AssetAPI;
 
 	TArray<int> MapToFirstOccurrences;
 
diff --git a/Engine/Plugins/Experimental/MeshModelingToolset/Source/MeshModelingTools/Public/CombineMeshesTool.h b/Engine/Plugins/Experimental/MeshModelingToolset/Source/MeshModelingTools/Public/CombineMeshesTool.h
--- a/Engine/Plugins/Experimental/MeshModelingToolset/Source/MeshModelingTools/Public/CombineMeshesTool.h
+++ b/Engine/Plugins/Experimental/MeshModelingToolset/Source/MeshModelingTools/Public/CombineMeshesTool.h
@@ -16,8 +16,6 @@
 // predeclarations
 struct FMeshDescription;
 class USimpleDynamicMeshComponent;
-class IAssetGenerationAPI;
-
 class IPrimitiveComponentBackedTarget;
 class IMeshDescriptionCommitter;
 class IMeshDescriptionProvider;
@@ -34,7 +32,6 @@ class MESHMODELINGTOOLS_API UCombineMeshesToolBuilder : public UInteractiveToolB
 	GENERATED_BODY()
 
 public:
-	IAssetGenerationAPI* AssetAPI = nullptr;
 	bool bIsDuplicateTool = false;
 
 	virtual bool CanBuildTool(const FToolBuilderState& SceneState) const override;
@@ -100,7 +97,6 @@ public:
 	virtual void Shutdown(EToolShutdownType ShutdownType) override;
 
 	virtual void SetWorld(UWorld* World);
-	virtual void SetAssetAPI(IAssetGenerationAPI* AssetAPI);
 
 	virtual bool HasCancel() const override { return true; }
 	virtual bool HasAccept() const override { return true; }
@@ -115,7 +111,6 @@ protected:
 
 protected:
 	UWorld* TargetWorld;
-	IAssetGenerationAPI* AssetAPI;
 
 	bool bDuplicateMode;
 
diff --git a/Engine/Plugins/Experimental/MeshModelingToolset/Source/MeshModelingTools/Public/DrawAndRevolveTool.h b/Engine/Plugins/Experimental/MeshModelingToolset/Source/MeshModelingTools/Public/DrawAndRevolveTool.h
--- a/Engine/Plugins/Experimental/MeshModelingToolset/Source/MeshModelingTools/Public/DrawAndRevolveTool.h
+++ b/Engine/Plugins/Experimental/MeshModelingToolset/Source/MeshModelingTools/Public/DrawAndRevolveTool.h
@@ -18,7 +18,6 @@
 class UCollectSurfacePathMechanic;
 class UConstructionPlaneMechanic;
 class UCurveControlPointsMechanic;
-class IAssetGenerationAPI;
 
 UCLASS()
 class MESHMODELINGTOOLS_API UDrawAndRevolveToolBuilder : public UInteractiveToolBuilder
@@ -26,8 +25,6 @@ class MESHMODELINGTOOLS_API UDrawAndRevolveToolBuilder : public UInteractiveTool
 	GENERATED_BODY()
 
 public:
-	IAssetGenerationAPI* AssetAPI = nullptr;
-
 	virtual bool CanBuildTool(const FToolBuilderState& SceneState) const override;
 	virtual UInteractiveTool* BuildTool(const FToolBuilderState& SceneState) const override;
 };
@@ -90,7 +87,6 @@ class MESHMODELINGTOOLS_API UDrawAndRevolveTool : public UInteractiveTool
 
 public:
 	virtual void SetWorld(UWorld* World) { TargetWorld = World; }
-	virtual void SetAssetAPI(IAssetGenerationAPI* NewAssetApi) { AssetAPI = NewAssetApi; }
 
 	virtual void RegisterActions(FInteractiveToolActionSet& ActionSet) override;
 	virtual void OnPointDeletionKeyPress();
@@ -110,7 +106,6 @@ public:
 protected:
 
 	UWorld* TargetWorld;
-	IAssetGenerationAPI* AssetAPI;
 
 	FViewCameraState CameraState;
 
diff --git a/Engine/Plugins/Experimental/MeshModelingToolset/Source/MeshModelingTools/Public/DrawPolyPathTool.h b/Engine/Plugins/Experimental/MeshModelingToolset/Source/MeshModelingTools/Public/DrawPolyPathTool.h
--- a/Engine/Plugins/Experimental/MeshModelingToolset/Source/MeshModelingTools/Public/DrawPolyPathTool.h
+++ b/Engine/Plugins/Experimental/MeshModelingToolset/Source/MeshModelingTools/Public/DrawPolyPathTool.h
@@ -16,7 +16,6 @@
 class FMeshVertexChangeBuilder;
 class UTransformGizmo;
 class UTransformProxy;
-class IAssetGenerationAPI;
 
 /**
  * ToolBuilder
@@ -26,8 +25,6 @@ class MESHMODELINGTOOLS_API UDrawPolyPathToolBuilder : public UMeshSurfacePointT
 {
 	GENERATED_BODY()
 public:
-	IAssetGenerationAPI* AssetAPI = nullptr;
-
 	virtual bool CanBuildTool(const FToolBuilderState& SceneState) const override;
 
 	virtual UInteractiveTool* BuildTool(const FToolBuilderState& SceneState) const override;
@@ -134,7 +131,6 @@ public:
 	virtual void RegisterActions(FInteractiveToolActionSet& ActionSet) override;
 
 	virtual void SetWorld(UWorld* World);
-	virtual void SetAssetAPI(IAssetGenerationAPI* AssetAPI);
 
 	virtual void Setup() override;
 	virtual void Shutdown(EToolShutdownType ShutdownType) override;
@@ -160,7 +156,6 @@ public:
 
 protected:
 	UWorld* TargetWorld;
-	IAssetGenerationAPI* AssetAPI;
 
 	UPROPERTY()
 	UDrawPolyPathProperties* TransformProps;
diff --git a/Engine/Plugins/Experimental/MeshModelingToolset/Source/MeshModelingTools/Public/DrawPolygonTool.h b/Engine/Plugins/Experimental/MeshModelingToolset/Source/MeshModelingTools/Public/DrawPolygonTool.h
--- a/Engine/Plugins/Experimental/MeshModelingToolset/Source/MeshModelingTools/Public/DrawPolygonTool.h
+++ b/Engine/Plugins/Experimental/MeshModelingToolset/Source/MeshModelingTools/Public/DrawPolygonTool.h
@@ -18,7 +18,6 @@
 
 class UTransformGizmo;
 class UTransformProxy;
-class IAssetGenerationAPI;
 PREDECLARE_USE_GEOMETRY_CLASS(FDynamicMesh3);
 
 /**
@@ -30,13 +29,6 @@ class MESHMODELINGTOOLS_API UDrawPolygonToolBuilder : public UInteractiveToolBui
 	GENERATED_BODY()
 
 public:
-	IAssetGenerationAPI* AssetAPI;
-
-	UDrawPolygonToolBuilder() 
-	{
-		AssetAPI = nullptr;
-	}
-
 	virtual bool CanBuildTool(const FToolBuilderState& SceneState) const override;
 	virtual UInteractiveTool* BuildTool(const FToolBuilderState& SceneState) const override;
 };
@@ -174,7 +166,6 @@ public:
 	virtual void RegisterActions(FInteractiveToolActionSet& ActionSet) override;
 
 	virtual void SetWorld(UWorld* World);
-	virtual void SetAssetAPI(IAssetGenerationAPI* AssetAPI);
 
 	virtual void Setup() override;
 	virtual void Shutdown(EToolShutdownType ShutdownType) override;
@@ -247,7 +238,6 @@ protected:
 
 protected:
 	UWorld* TargetWorld;
-	IAssetGenerationAPI* AssetAPI;
 
 	FViewCameraState CameraState;
 
diff --git a/Engine/Plugins/Experimental/MeshModelingToolset/Source/MeshModelingTools/Public/EditNormalsTool.h b/Engine/Plugins/Experimental/MeshModelingToolset/Source/MeshModelingTools/Public/EditNormalsTool.h
--- a/Engine/Plugins/Experimental/MeshModelingToolset/Source/MeshModelingTools/Public/EditNormalsTool.h
+++ b/Engine/Plugins/Experimental/MeshModelingToolset/Source/MeshModelingTools/Public/EditNormalsTool.h
@@ -16,10 +16,6 @@
 // predeclarations
 struct FMeshDescription;
 class USimpleDynamicMeshComponent;
-class IAssetGenerationAPI;
-
-
-
 
 /**
  *
@@ -30,8 +26,6 @@ class MESHMODELINGTOOLS_API UEditNormalsToolBuilder : public UInteractiveToolBui
 	GENERATED_BODY()
 
 public:
-	IAssetGenerationAPI* AssetAPI = nullptr;
-
 	virtual bool CanBuildTool(const FToolBuilderState& SceneState) const override;
 	virtual UInteractiveTool* BuildTool(const FToolBuilderState& SceneState) const override;
 
@@ -142,7 +136,6 @@ public:
 	virtual void Shutdown(EToolShutdownType ShutdownType) override;
 
 	virtual void SetWorld(UWorld* World);
-	virtual void SetAssetAPI(IAssetGenerationAPI* AssetAPI);
 
 	virtual void OnTick(float DeltaTime) override;
 	virtual void Render(IToolsContextRenderAPI* RenderAPI) override;
@@ -174,7 +167,6 @@ protected:
 	TArray<TSharedPtr<UE::Geometry::FDynamicMesh3, ESPMode::ThreadSafe>> OriginalDynamicMeshes;
 
 	UWorld* TargetWorld;
-	IAssetGenerationAPI* AssetAPI;
 
 	FViewCameraState CameraState;
 
diff --git a/Engine/Plugins/Experimental/MeshModelingToolset/Source/MeshModelingTools/Public/MeshAttributePaintTool.h b/Engine/Plugins/Experimental/MeshModelingToolset/Source/MeshModelingTools/Public/MeshAttributePaintTool.h
--- a/Engine/Plugins/Experimental/MeshModelingToolset/Source/MeshModelingTools/Public/MeshAttributePaintTool.h
+++ b/Engine/Plugins/Experimental/MeshModelingToolset/Source/MeshModelingTools/Public/MeshAttributePaintTool.h
@@ -85,10 +85,7 @@ UCLASS()
 class MESHMODELINGTOOLS_API UMeshAttributePaintToolBuilder : public UMeshSurfacePointToolBuilder
 {
 	GENERATED_BODY()
-
 public:
-	IToolsContextAssetAPI* AssetAPI = nullptr;
-
 	/** Optional color map customization */
 	TUniqueFunction<TUniquePtr<FFloatAttributeColorMapper>()> ColorMapFactory;
 
@@ -193,7 +190,6 @@ class MESHMODELINGTOOLS_API UMeshAttributePaintTool : public UDynamicMeshBrushTo
 
 public:
 	virtual void SetWorld(UWorld* World);
-	virtual void SetAssetAPI(IToolsContextAssetAPI* AssetAPI);
 
 	virtual void RegisterActions(FInteractiveToolActionSet& ActionSet) override;
 
@@ -245,7 +241,6 @@ protected:
 
 protected:
 	UWorld* TargetWorld;
-	IToolsContextAssetAPI* AssetAPI;
 
 	bool bInRemoveStroke = false;
 	bool bInSmoothStroke = false;
diff --git a/Engine/Plugins/Experimental/MeshModelingToolset/Source/MeshModelingTools/Public/MeshSelectionTool.h b/Engine/Plugins/Experimental/MeshModelingToolset/Source/MeshModelingTools/Public/MeshSelectionTool.h
--- a/Engine/Plugins/Experimental/MeshModelingToolset/Source/MeshModelingTools/Public/MeshSelectionTool.h
+++ b/Engine/Plugins/Experimental/MeshModelingToolset/Source/MeshModelingTools/Public/MeshSelectionTool.h
@@ -9,8 +9,6 @@
 #include "DynamicMeshOctree3.h"
 #include "MeshSelectionTool.generated.h"
 
-class IAssetGenerationAPI;
-
 /**
  *
  */
@@ -20,8 +18,6 @@ class MESHMODELINGTOOLS_API UMeshSelectionToolBuilder : public UMeshSurfacePoint
 	GENERATED_BODY()
 
 public:
-	IAssetGenerationAPI* AssetAPI = nullptr;
-
 	virtual UMeshSurfacePointTool* CreateNewTool(const FToolBuilderState& SceneState) const override;
 };
 
@@ -271,7 +267,6 @@ public:
 	UMeshSelectionTool();
 
 	virtual void SetWorld(UWorld* World);
-	virtual void SetAssetAPI(IAssetGenerationAPI* AssetAPI);
 
 	virtual void RegisterActions(FInteractiveToolActionSet& ActionSet) override;
 
@@ -326,7 +321,6 @@ protected:
 	TArray<AActor*> SpawnedActors;
 
 	UWorld* TargetWorld;
-	IAssetGenerationAPI* AssetAPI;
 
 	// note: ideally this octree would be part of PreviewMesh!
 	TUniquePtr<UE::Geometry::FDynamicMeshOctree3> Octree;
diff --git a/Engine/Plugins/Experimental/MeshModelingToolset/Source/MeshModelingTools/Public/MirrorTool.h b/Engine/Plugins/Experimental/MeshModelingToolset/Source/MeshModelingTools/Public/MirrorTool.h
--- a/Engine/Plugins/Experimental/MeshModelingToolset/Source/MeshModelingTools/Public/MirrorTool.h
+++ b/Engine/Plugins/Experimental/MeshModelingToolset/Source/MeshModelingTools/Public/MirrorTool.h
@@ -17,16 +17,12 @@
 
 class UTransformGizmo;
 class UTransformProxy;
-class IAssetGenerationAPI;
 
 UCLASS()
 class MESHMODELINGTOOLS_API UMirrorToolBuilder : public UInteractiveToolBuilder
 {
 	GENERATED_BODY()
-
 public:
-	IAssetGenerationAPI* AssetAPI = nullptr;
-
 	virtual bool CanBuildTool(const FToolBuilderState& SceneState) const override;
 	virtual UInteractiveTool* BuildTool(const FToolBuilderState& SceneState) const override;
 
@@ -198,7 +194,6 @@ public:
 	virtual void Shutdown(EToolShutdownType ShutdownType) override;
 
 	virtual void SetWorld(UWorld* World);
-	virtual void SetAssetAPI(IAssetGenerationAPI* AssetAPI);
 
 	virtual void OnTick(float DeltaTime) override;
 	virtual void Render(IToolsContextRenderAPI* RenderAPI) override;
@@ -231,7 +226,6 @@ protected:
 	FVector3d MirrorPlaneNormal = FVector3d::UnitZ();
 
 	UWorld* TargetWorld;
-	IAssetGenerationAPI* AssetAPI;
 
 	// flags used to identify modifier keys/buttons
 	static const int SnappingToggleModifierId = 1;
diff --git a/Engine/Plugins/Experimental/MeshModelingToolset/Source/MeshModelingTools/Public/Physics/ExtractCollisionGeometryTool.h b/Engine/Plugins/Experimental/MeshModelingToolset/Source/MeshModelingTools/Public/Physics/ExtractCollisionGeometryTool.h
--- a/Engine/Plugins/Experimental/MeshModelingToolset/Source/MeshModelingTools/Public/Physics/ExtractCollisionGeometryTool.h
+++ b/Engine/Plugins/Experimental/MeshModelingToolset/Source/MeshModelingTools/Public/Physics/ExtractCollisionGeometryTool.h
@@ -12,7 +12,6 @@
 
 class UPreviewGeometry;
 class UPreviewMesh;
-class IAssetGenerationAPI;
 
 UCLASS()
 class MESHMODELINGTOOLS_API UExtractCollisionGeometryToolBuilder : public USingleSelectionMeshEditingToolBuilder
@@ -21,7 +20,6 @@ class MESHMODELINGTOOLS_API UExtractCollisionGeometryToolBuilder : public USingl
 
 public:
 	virtual USingleSelectionMeshEditingTool* CreateNewTool(const FToolBuilderState& SceneState) const override;
-	virtual bool RequiresAssetAPI() const override { return true; }
 
 protected:
 	virtual const FToolTargetTypeRequirements& GetTargetRequirements() const override;
diff --git a/Engine/Plugins/Experimental/MeshModelingToolset/Source/MeshModelingTools/Public/PlaneCutTool.h b/Engine/Plugins/Experimental/MeshModelingToolset/Source/MeshModelingTools/Public/PlaneCutTool.h
--- a/Engine/Plugins/Experimental/MeshModelingToolset/Source/MeshModelingTools/Public/PlaneCutTool.h
+++ b/Engine/Plugins/Experimental/MeshModelingToolset/Source/MeshModelingTools/Public/PlaneCutTool.h
@@ -3,7 +3,6 @@
 #pragma once
 
 #include "CoreMinimal.h"
-#include "UObject/NoExportTypes.h"
 #include "InteractiveGizmo.h"
 #include "MultiSelectionTool.h"
 #include "InteractiveToolBuilder.h"
@@ -18,9 +17,6 @@
 // predeclarations
 class UTransformGizmo;
 class UTransformProxy;
-class IAssetGenerationAPI;
-
-
 
 
 /**
@@ -32,8 +28,6 @@ class MESHMODELINGTOOLS_API UPlaneCutToolBuilder : public UInteractiveToolBuilde
 	GENERATED_BODY()
 
 public:
-	IAssetGenerationAPI* AssetAPI = nullptr;
-
 	virtual bool CanBuildTool(const FToolBuilderState& SceneState) const override;
 	virtual UInteractiveTool* BuildTool(const FToolBuilderState& SceneState) const override;
 
@@ -149,7 +143,6 @@ public:
 	virtual void Shutdown(EToolShutdownType ShutdownType) override;
 
 	virtual void SetWorld(UWorld* World);
-	virtual void SetAssetAPI(IAssetGenerationAPI* AssetAPI);
 
 	virtual void RegisterActions(FInteractiveToolActionSet& ActionSet) override;
 
@@ -215,7 +208,6 @@ protected:
 	TArray<float> MeshUVScaleFactor;
 
 	UWorld* TargetWorld;
-	IAssetGenerationAPI* AssetAPI;
 
 	FViewCameraState CameraState;
 
diff --git a/Engine/Plugins/Experimental/MeshModelingToolset/Source/MeshModelingTools/Public/ProjectToTargetTool.h b/Engine/Plugins/Experimental/MeshModelingToolset/Source/MeshModelingTools/Public/ProjectToTargetTool.h
--- a/Engine/Plugins/Experimental/MeshModelingToolset/Source/MeshModelingTools/Public/ProjectToTargetTool.h
+++ b/Engine/Plugins/Experimental/MeshModelingToolset/Source/MeshModelingTools/Public/ProjectToTargetTool.h
@@ -13,12 +13,8 @@ UCLASS()
 class MESHMODELINGTOOLS_API UProjectToTargetToolBuilder : public UInteractiveToolBuilder
 {
 	GENERATED_BODY()
-
 public:
-	IToolsContextAssetAPI* AssetAPI = nullptr;
-
 	virtual bool CanBuildTool(const FToolBuilderState& SceneState) const override;
-
 	virtual UInteractiveTool* BuildTool(const FToolBuilderState& SceneState) const override;
 
 protected:
diff --git a/Engine/Plugins/Experimental/MeshModelingToolset/Source/MeshModelingTools/Public/RemeshMeshTool.h b/Engine/Plugins/Experimental/MeshModelingToolset/Source/MeshModelingTools/Public/RemeshMeshTool.h
--- a/Engine/Plugins/Experimental/MeshModelingToolset/Source/MeshModelingTools/Public/RemeshMeshTool.h
+++ b/Engine/Plugins/Experimental/MeshModelingToolset/Source/MeshModelingTools/Public/RemeshMeshTool.h
@@ -24,8 +24,6 @@ class MESHMODELINGTOOLS_API URemeshMeshToolBuilder : public UInteractiveToolBuil
 	GENERATED_BODY()
 
 public:
-	IToolsContextAssetAPI* AssetAPI = nullptr;
-
 	/** 
 	 * Return true if we have one object selected. URemeshMeshTool is a UMultiSelectionTool, however we currently 
 	 * only ever apply it to a single mesh. (See comment at URemeshMeshTool definition below.)
@@ -122,7 +120,6 @@ public:
 	URemeshMeshTool(const FObjectInitializer&);
 
 	virtual void SetWorld(UWorld* World);
-	virtual void SetAssetAPI(IToolsContextAssetAPI* AssetAPI);
 
 	virtual void Setup() override;
 	virtual void Shutdown(EToolShutdownType ShutdownType) override;
@@ -151,7 +148,6 @@ public:
 protected:
 
 	UWorld* TargetWorld;
-	IToolsContextAssetAPI* AssetAPI;
 
 	TSharedPtr<UE::Geometry::FDynamicMesh3, ESPMode::ThreadSafe> OriginalMesh;
 	TSharedPtr<UE::Geometry::FDynamicMeshAABBTree3, ESPMode::ThreadSafe> OriginalMeshSpatial;
diff --git a/Engine/Plugins/Experimental/MeshModelingToolset/Source/MeshModelingTools/Public/RemoveOccludedTrianglesTool.h b/Engine/Plugins/Experimental/MeshModelingToolset/Source/MeshModelingTools/Public/RemoveOccludedTrianglesTool.h
--- a/Engine/Plugins/Experimental/MeshModelingToolset/Source/MeshModelingTools/Public/RemoveOccludedTrianglesTool.h
+++ b/Engine/Plugins/Experimental/MeshModelingToolset/Source/MeshModelingTools/Public/RemoveOccludedTrianglesTool.h
@@ -15,7 +15,6 @@
 
 #include "RemoveOccludedTrianglesTool.generated.h"
 
-class IAssetGenerationAPI;
 PREDECLARE_USE_GEOMETRY_CLASS(FDynamicMesh3);
 
 /**
@@ -27,8 +26,6 @@ class MESHMODELINGTOOLS_API URemoveOccludedTrianglesToolBuilder : public UIntera
 	GENERATED_BODY()
 
 public:
-	IAssetGenerationAPI* AssetAPI = nullptr;
-
 	virtual bool CanBuildTool(const FToolBuilderState& SceneState) const override;
 	virtual UInteractiveTool* BuildTool(const FToolBuilderState& SceneState) const override;
 
@@ -170,7 +167,6 @@ public:
 	virtual void Shutdown(EToolShutdownType ShutdownType) override;
 
 	virtual void SetWorld(UWorld* World);
-	virtual void SetAssetAPI(IAssetGenerationAPI* AssetAPI);
 
 	virtual void OnTick(float DeltaTime) override;
 
@@ -222,7 +218,6 @@ protected:
 	TArray<int32> OccludedGroupLayers;
 
 	UWorld* TargetWorld;
-	IAssetGenerationAPI* AssetAPI;
 
 	FViewCameraState CameraState;
 
diff --git a/Engine/Plugins/Experimental/MeshModelingToolset/Source/MeshModelingTools/Public/RevolveBoundaryTool.h b/Engine/Plugins/Experimental/MeshModelingToolset/Source/MeshModelingTools/Public/RevolveBoundaryTool.h
--- a/Engine/Plugins/Experimental/MeshModelingToolset/Source/MeshModelingTools/Public/RevolveBoundaryTool.h
+++ b/Engine/Plugins/Experimental/MeshModelingToolset/Source/MeshModelingTools/Public/RevolveBoundaryTool.h
@@ -19,8 +19,6 @@
 
 #include "RevolveBoundaryTool.generated.h"
 
-class IAssetGenerationAPI;
-
 // Tool Builder
 
 UCLASS()
@@ -30,7 +28,6 @@ class MESHMODELINGTOOLS_API URevolveBoundaryToolBuilder : public USingleSelectio
 
 public:
 	virtual USingleSelectionMeshEditingTool* CreateNewTool(const FToolBuilderState& SceneState) const override;
-	virtual bool RequiresAssetAPI() const override { return true; }
 };
 
 UCLASS()
diff --git a/Engine/Plugins/Experimental/MeshModelingToolset/Source/MeshModelingTools/Public/SkinWeightsBindingTool.h b/Engine/Plugins/Experimental/MeshModelingToolset/Source/MeshModelingTools/Public/SkinWeightsBindingTool.h
--- a/Engine/Plugins/Experimental/MeshModelingToolset/Source/MeshModelingTools/Public/SkinWeightsBindingTool.h
+++ b/Engine/Plugins/Experimental/MeshModelingToolset/Source/MeshModelingTools/Public/SkinWeightsBindingTool.h
@@ -16,8 +16,6 @@ class UMeshOpPreviewWithBackgroundCompute;
 struct FDynamicMeshOpResult;
 struct FMeshDescription;
 class IMeshDescriptionCommitter;
-
-
 struct FOccupancyGrid;
 
 /**
@@ -27,12 +25,8 @@ UCLASS()
 class MESHMODELINGTOOLS_API USkinWeightsBindingToolBuilder : public UInteractiveToolBuilder
 {
 	GENERATED_BODY()
-
 public:
-	IToolsContextAssetAPI* AssetAPI = nullptr;
-
 	bool CanBuildTool(const FToolBuilderState& SceneState) const override;
-
 	UInteractiveTool* BuildTool(const FToolBuilderState& SceneState) const override;
 };
 
diff --git a/Engine/Plugins/Experimental/MeshModelingToolset/Source/MeshModelingTools/Public/UVProjectionTool.h b/Engine/Plugins/Experimental/MeshModelingToolset/Source/MeshModelingTools/Public/UVProjectionTool.h
--- a/Engine/Plugins/Experimental/MeshModelingToolset/Source/MeshModelingTools/Public/UVProjectionTool.h
+++ b/Engine/Plugins/Experimental/MeshModelingToolset/Source/MeshModelingTools/Public/UVProjectionTool.h
@@ -11,6 +11,7 @@
 #include "Properties/MeshUVChannelProperties.h"
 #include "Drawing/PreviewGeometryActor.h"
 #include "Selection/SelectClickedAction.h"
+#include "OrientedBoxTypes.h"
 
 #include "UVProjectionTool.generated.h"
 
@@ -21,7 +22,6 @@ class USimpleDynamicMeshComponent;
 class UTransformGizmo;
 class UTransformProxy;
 class USingleClickInputBehavior;
-class IAssetGenerationAPI;
 
 /**
  *
@@ -30,7 +30,6 @@ UCLASS()
 class MESHMODELINGTOOLS_API UUVProjectionToolBuilder : public USingleSelectionMeshEditingToolBuilder
 {
 	GENERATED_BODY()
-
 public:
 	virtual USingleSelectionMeshEditingTool* CreateNewTool(const FToolBuilderState& SceneState) const override;
 
diff --git a/Engine/Plugins/Experimental/MeshModelingToolset/Source/MeshModelingTools/Public/VolumeToMeshTool.h b/Engine/Plugins/Experimental/MeshModelingToolset/Source/MeshModelingTools/Public/VolumeToMeshTool.h
--- a/Engine/Plugins/Experimental/MeshModelingToolset/Source/MeshModelingTools/Public/VolumeToMeshTool.h
+++ b/Engine/Plugins/Experimental/MeshModelingToolset/Source/MeshModelingTools/Public/VolumeToMeshTool.h
@@ -11,8 +11,6 @@
 #include "Drawing/LineSetComponent.h"
 #include "VolumeToMeshTool.generated.h"
 
-class IAssetGenerationAPI;
-
 /**
  *
  */
@@ -21,8 +19,6 @@ class MESHMODELINGTOOLS_API UVolumeToMeshToolBuilder : public UInteractiveToolBu
 {
 	GENERATED_BODY()
 public:
-	IAssetGenerationAPI* AssetAPI = nullptr;
-
 	virtual bool CanBuildTool(const FToolBuilderState& SceneState) const override;
 	virtual UInteractiveTool* BuildTool(const FToolBuilderState& SceneState) const override;
 };
@@ -74,7 +70,6 @@ public:
 	UVolumeToMeshTool();
 
 	virtual void SetWorld(UWorld* World) { this->TargetWorld = World; }
-	virtual void SetAssetAPI(IAssetGenerationAPI* InAssetAPI) { this->AssetAPI = InAssetAPI; }
 	virtual void SetSelection(AVolume* Volume);
 
 	virtual void Setup() override;
@@ -103,7 +98,6 @@ protected:
 
 protected:
 	UWorld* TargetWorld = nullptr;
-	IAssetGenerationAPI* AssetAPI = nullptr;
 
 	UE::Geometry::FDynamicMesh3 CurrentMesh;
 
diff --git a/Engine/Plugins/Experimental/MeshModelingToolset/Source/MeshModelingToolsEditorOnly/Private/BspConversionTool.cpp b/Engine/Plugins/Experimental/MeshModelingToolset/Source/MeshModelingToolsEditorOnly/Private/BspConversionTool.cpp
--- a/Engine/Plugins/Experimental/MeshModelingToolset/Source/MeshModelingToolsEditorOnly/Private/BspConversionTool.cpp
+++ b/Engine/Plugins/Experimental/MeshModelingToolset/Source/MeshModelingToolsEditorOnly/Private/BspConversionTool.cpp
@@ -2,7 +2,8 @@
 
 #include "BspConversionTool.h"
 
-#include "AssetGenerationUtil.h"
+#include "ModelingObjectsCreationAPI.h"
+#include "Selection/ToolSelectionUtil.h"
 #include "AssetSelection.h"
 #include "ComponentSourceInterfaces.h"
 #include "DynamicMesh3.h"
@@ -49,15 +50,13 @@ bool UBspConversionToolBuilder::CanBuildTool(const FToolBuilderState& SceneState
 {
 	// We allow the tool to be built even if nothing is selected because the tool has a "select all" option.
 	// We just need to be able to get the place to save assets.
-	return AssetAPI != nullptr;
+	return true;
 }
 
 UInteractiveTool* UBspConversionToolBuilder::BuildTool(const FToolBuilderState& SceneState) const
 {
 	UBspConversionTool* NewTool = NewObject<UBspConversionTool>(SceneState.ToolManager);
 	NewTool->SetWorld(SceneState.World);
-	NewTool->SetAssetAPI(AssetAPI);
-
 	return NewTool;
 }
 
@@ -212,19 +211,16 @@ void UBspConversionTool::Shutdown(EToolShutdownType ShutdownType)
 			TArray<UMaterialInterface*> Materials;
 			PreviewMesh->GetMaterials(Materials);
 
-			// Create the final actor
-			AActor* NewActor = AssetGenerationUtil::GenerateStaticMeshActor(
-				AssetAPI, TargetWorld,
-				PreviewMesh->GetMesh(), (UE::Geometry::FTransform3d)PreviewMesh->GetTransform(),
-				TEXT("BspConvertedToStaticMesh"), Materials);
-
-			if (NewActor)
+			FCreateMeshObjectParams NewMeshObjectParams;
+			NewMeshObjectParams.TargetWorld = TargetWorld;
+			NewMeshObjectParams.Transform = PreviewMesh->GetTransform();
+			NewMeshObjectParams.BaseName = TEXT("BspMesh");
+			NewMeshObjectParams.Materials = Materials;
+			NewMeshObjectParams.SetMesh(PreviewMesh->GetMesh());
+			FCreateMeshObjectResult Result = UE::Modeling::CreateMeshObject(GetToolManager(), MoveTemp(NewMeshObjectParams));
+			if (Result.IsOK() && Result.NewActor != nullptr)
 			{
-				// Select the final result
-				FSelectedOjectsChangeList NewSelection;
-				NewSelection.ModificationType = ESelectedObjectsModificationType::Replace;
-				NewSelection.Actors.Add(NewActor);
-				GetToolManager()->RequestSelectionChange(NewSelection);
+				ToolSelectionUtil::SetNewActorSelection(GetToolManager(), Result.NewActor);
 			}
 		}
 
diff --git a/Engine/Plugins/Experimental/MeshModelingToolset/Source/MeshModelingToolsEditorOnly/Private/EditMeshMaterialsTool.cpp b/Engine/Plugins/Experimental/MeshModelingToolset/Source/MeshModelingToolsEditorOnly/Private/EditMeshMaterialsTool.cpp
--- a/Engine/Plugins/Experimental/MeshModelingToolset/Source/MeshModelingToolsEditorOnly/Private/EditMeshMaterialsTool.cpp
+++ b/Engine/Plugins/Experimental/MeshModelingToolset/Source/MeshModelingToolsEditorOnly/Private/EditMeshMaterialsTool.cpp
@@ -12,7 +12,6 @@
 #include "Selections/MeshConnectedComponents.h"
 #include "MeshRegionBoundaryLoops.h"
 #include "DynamicMesh/MeshIndexUtil.h"
-#include "AssetGenerationUtil.h"
 #include "ToolSetupUtil.h"
 
 #include "TargetInterfaces/MaterialProvider.h"
@@ -41,7 +40,6 @@ UMeshSurfacePointTool* UEditMeshMaterialsToolBuilder::CreateNewTool(const FToolB
 {
 	UEditMeshMaterialsTool* SelectionTool = NewObject<UEditMeshMaterialsTool>(SceneState.ToolManager);
 	SelectionTool->SetWorld(SceneState.World);
-	SelectionTool->SetAssetAPI(AssetAPI);
 	return SelectionTool;
 }
 
diff --git a/Engine/Plugins/Experimental/MeshModelingToolset/Source/MeshModelingToolsEditorOnly/Private/GenerateLODMeshesTool.cpp b/Engine/Plugins/Experimental/MeshModelingToolset/Source/MeshModelingToolsEditorOnly/Private/GenerateLODMeshesTool.cpp
--- a/Engine/Plugins/Experimental/MeshModelingToolset/Source/MeshModelingToolsEditorOnly/Private/GenerateLODMeshesTool.cpp
+++ b/Engine/Plugins/Experimental/MeshModelingToolset/Source/MeshModelingToolsEditorOnly/Private/GenerateLODMeshesTool.cpp
@@ -13,7 +13,8 @@
 #include "MeshDescriptionToDynamicMesh.h"
 #include "DynamicMeshToMeshDescription.h"
 
-#include "AssetGenerationUtil.h"
+#include "ModelingObjectsCreationAPI.h"
+#include "Selection/ToolSelectionUtil.h"
 
 #include "SceneManagement.h" // for FPrimitiveDrawInterface
 
@@ -381,8 +382,8 @@ void UGenerateLODMeshesTool::GenerateAssets()
 	{
 		UMeshOpPreviewWithBackgroundCompute* Preview = Previews[k];
 
-		FDynamicMeshOpResult Result = Preview->Shutdown();
-		if (Result.Mesh->TriangleCount() == 0)
+		FDynamicMeshOpResult OpResult = Preview->Shutdown();
+		if (OpResult.Mesh->TriangleCount() == 0)
 		{
 			continue;		// failed!
 		}
@@ -390,12 +391,20 @@ void UGenerateLODMeshesTool::GenerateAssets()
 		FComponentMaterialSet MaterialSet;
 		Cast<IMaterialProvider>(Target)->GetMaterialSet(MaterialSet);
 
-		FString BaseName = AssetGenerationUtil::GetComponentAssetBaseName(TargetComponent->GetOwnerComponent());
+		FString BaseName = UE::Modeling::GetComponentAssetBaseName(TargetComponent->GetOwnerComponent());
 		FString Name = FString::Printf( TEXT("%s_LOD%d"), *BaseName, (SimplifyProperties->NameIndexBase+k) );
 
-		AActor* NewActor = AssetGenerationUtil::GenerateStaticMeshActor(
-			AssetAPI, TargetWorld,
-			Result.Mesh.Get(), UE::Geometry::FTransform3d(Transform), Name, MaterialSet.Materials);
+		FCreateMeshObjectParams NewMeshObjectParams;
+		NewMeshObjectParams.TargetWorld = TargetWorld;
+		NewMeshObjectParams.Transform = Transform;
+		NewMeshObjectParams.BaseName = Name;
+		NewMeshObjectParams.Materials = MaterialSet.Materials;
+		NewMeshObjectParams.SetMesh(OpResult.Mesh.Get());
+		FCreateMeshObjectResult Result = UE::Modeling::CreateMeshObject(GetToolManager(), MoveTemp(NewMeshObjectParams));
+		if (Result.IsOK() && Result.NewActor != nullptr)
+		{
+			ToolSelectionUtil::SetNewActorSelection(GetToolManager(), Result.NewActor);
+		}
 	}
 
 	GetToolManager()->EndUndoTransaction();
diff --git a/Engine/Plugins/Experimental/MeshModelingToolset/Source/MeshModelingToolsEditorOnly/Private/Hair/GroomCardsEditorTool.cpp b/Engine/Plugins/Experimental/MeshModelingToolset/Source/MeshModelingToolsEditorOnly/Private/Hair/GroomCardsEditorTool.cpp
--- a/Engine/Plugins/Experimental/MeshModelingToolset/Source/MeshModelingToolsEditorOnly/Private/Hair/GroomCardsEditorTool.cpp
+++ b/Engine/Plugins/Experimental/MeshModelingToolset/Source/MeshModelingToolsEditorOnly/Private/Hair/GroomCardsEditorTool.cpp
@@ -14,7 +14,6 @@
 #include "Util/ColorConstants.h"
 #include "ToolSetupUtil.h"
 #include "Selection/ToolSelectionUtil.h"
-#include "AssetGenerationUtil.h"
 #include "MeshTransforms.h"
 #include "DynamicMeshEditor.h"
 #include "DynamicMeshChangeTracker.h"
@@ -47,7 +46,6 @@ UMeshSurfacePointTool* UGroomCardsEditorToolBuilder::CreateNewTool(const FToolBu
 {
 	UGroomCardsEditorTool* NewTool = NewObject<UGroomCardsEditorTool>(SceneState.ToolManager);
 	NewTool->SetWorld(SceneState.World);
-	NewTool->SetAssetAPI(AssetAPI);
 	return NewTool;
 }
 
diff --git a/Engine/Plugins/Experimental/MeshModelingToolset/Source/MeshModelingToolsEditorOnly/Private/Hair/GroomToMeshTool.cpp b/Engine/Plugins/Experimental/MeshModelingToolset/Source/MeshModelingToolsEditorOnly/Private/Hair/GroomToMeshTool.cpp
--- a/Engine/Plugins/Experimental/MeshModelingToolset/Source/MeshModelingToolsEditorOnly/Private/Hair/GroomToMeshTool.cpp
+++ b/Engine/Plugins/Experimental/MeshModelingToolset/Source/MeshModelingToolsEditorOnly/Private/Hair/GroomToMeshTool.cpp
@@ -12,7 +12,7 @@
 #include "Util/ColorConstants.h"
 #include "ToolSetupUtil.h"
 #include "Selection/ToolSelectionUtil.h"
-#include "AssetGenerationUtil.h"
+#include "ModelingObjectsCreationAPI.h"
 #include "MeshTransforms.h"
 #include "Selections/MeshFaceSelection.h"
 #include "Selections/MeshVertexSelection.h"
@@ -68,9 +68,6 @@ UInteractiveTool* UGroomToMeshToolBuilder::BuildTool(const FToolBuilderState& Sc
 	UGroomToMeshTool* NewTool = NewObject<UGroomToMeshTool>(SceneState.ToolManager);
 
 	NewTool->SetWorld(SceneState.World);
-	check(AssetAPI);
-	NewTool->SetAssetAPI(AssetAPI);
-
 	AGroomActor* Groom = ToolBuilderUtil::FindFirstActorOfType<AGroomActor>(SceneState);
 	check(Groom != nullptr);
 	NewTool->SetSelection(Groom);
@@ -171,12 +168,16 @@ void UGroomToMeshTool::Shutdown(EToolShutdownType ShutdownType)
 
 		GetToolManager()->BeginUndoTransaction(LOCTEXT("CreateMeshGroom", "Groom To Mesh"));
 
-		AActor* NewActor = AssetGenerationUtil::GenerateStaticMeshActor(
-			AssetAPI, TargetWorld,
-			&CurrentMesh, Transform, NewName, UseMaterial);
-		if (NewActor != nullptr)
+		FCreateMeshObjectParams NewMeshObjectParams;
+		NewMeshObjectParams.TargetWorld = TargetWorld;
+		NewMeshObjectParams.Transform = (FTransform)Transform;
+		NewMeshObjectParams.BaseName = NewName;
+		NewMeshObjectParams.Materials.Add(UseMaterial);
+		NewMeshObjectParams.SetMesh(&CurrentMesh);
+		FCreateMeshObjectResult Result = UE::Modeling::CreateMeshObject(GetToolManager(), MoveTemp(NewMeshObjectParams));
+		if (Result.IsOK() && Result.NewActor != nullptr)
 		{
-			ToolSelectionUtil::SetNewActorSelection(GetToolManager(), NewActor);
+			ToolSelectionUtil::SetNewActorSelection(GetToolManager(), Result.NewActor);
 		}
 
 		GetToolManager()->EndUndoTransaction();
diff --git a/Engine/Plugins/Experimental/MeshModelingToolset/Source/MeshModelingToolsEditorOnly/Private/MergeMeshesTool.cpp b/Engine/Plugins/Experimental/MeshModelingToolset/Source/MeshModelingToolsEditorOnly/Private/MergeMeshesTool.cpp
--- a/Engine/Plugins/Experimental/MeshModelingToolset/Source/MeshModelingToolsEditorOnly/Private/MergeMeshesTool.cpp
+++ b/Engine/Plugins/Experimental/MeshModelingToolset/Source/MeshModelingToolsEditorOnly/Private/MergeMeshesTool.cpp
@@ -4,7 +4,7 @@
 #include "InteractiveToolManager.h"
 #include "ToolBuilderUtil.h"
 #include "ToolSetupUtil.h"
-#include "AssetGenerationUtil.h"
+#include "ModelingObjectsCreationAPI.h"
 #include "Selection/ToolSelectionUtil.h"
 
 #include "DynamicMesh3.h"
@@ -37,10 +37,9 @@ const FToolTargetTypeRequirements& UMergeMeshesToolBuilder::GetTargetRequirement
 
 bool UMergeMeshesToolBuilder::CanBuildTool(const FToolBuilderState& SceneState) const
 {
-	const bool bHasBuildAPI = (this->AssetAPI != nullptr);
 	const int32 MinRequiredComponents = 1;
 	const bool bHasComponents = SceneState.TargetManager->CountSelectedAndTargetable(SceneState, GetTargetRequirements()) >= MinRequiredComponents;
-	return ( bHasBuildAPI && bHasComponents );
+	return ( bHasComponents );
 }
 
 UInteractiveTool* UMergeMeshesToolBuilder::BuildTool(const FToolBuilderState& SceneState) const
@@ -50,7 +49,6 @@ UInteractiveTool* UMergeMeshesToolBuilder::BuildTool(const FToolBuilderState& Sc
 	TArray<TObjectPtr<UToolTarget>> Targets = SceneState.TargetManager->BuildAllSelectedTargetable(SceneState, GetTargetRequirements());
 	NewTool->SetTargets(MoveTemp(Targets));
 	NewTool->SetWorld(SceneState.World);
-	NewTool->SetAssetAPI(AssetAPI);
 
 	return NewTool;
 }
@@ -71,12 +69,6 @@ void UMergeMeshesTool::SetWorld(UWorld* World)
 	this->TargetWorld = World;
 }
 
-void UMergeMeshesTool::SetAssetAPI(IAssetGenerationAPI* AssetAPIIn)
-{
-	this->AssetAPI = AssetAPIIn;
-}
-
-
 
 void UMergeMeshesTool::Setup()
 {
@@ -224,20 +216,22 @@ void UMergeMeshesTool::CreateLowQualityPreview()
 }
 
 
-void UMergeMeshesTool::GenerateAsset(const FDynamicMeshOpResult& Result)
+void UMergeMeshesTool::GenerateAsset(const FDynamicMeshOpResult& OpResult)
 {
-	check(Result.Mesh.Get() != nullptr);
-
+	check(OpResult.Mesh.Get() != nullptr);
 	
-
-	AActor* NewActor = AssetGenerationUtil::GenerateStaticMeshActor(
-		AssetAPI, TargetWorld,
-		Result.Mesh.Get(), Result.Transform, TEXT("MergedMesh"));
-	if (NewActor != nullptr)
+	FCreateMeshObjectParams NewMeshObjectParams;
+	NewMeshObjectParams.TargetWorld = TargetWorld;
+	NewMeshObjectParams.Transform = (FTransform)OpResult.Transform;
+	NewMeshObjectParams.BaseName = TEXT("MergedMesh");
+	NewMeshObjectParams.Materials.Add(ToolSetupUtil::GetDefaultMaterial());
+	NewMeshObjectParams.SetMesh(OpResult.Mesh.Get());
+	FCreateMeshObjectResult Result = UE::Modeling::CreateMeshObject(GetToolManager(), MoveTemp(NewMeshObjectParams));
+	if (Result.IsOK() && Result.NewActor != nullptr)
 	{
-		ToolSelectionUtil::SetNewActorSelection(GetToolManager(), NewActor);
+		ToolSelectionUtil::SetNewActorSelection(GetToolManager(), Result.NewActor);
 	}
-	
+
 }
 
 
diff --git a/Engine/Plugins/Experimental/MeshModelingToolset/Source/MeshModelingToolsEditorOnly/Private/ShapeSprayTool.cpp b/Engine/Plugins/Experimental/MeshModelingToolset/Source/MeshModelingToolsEditorOnly/Private/ShapeSprayTool.cpp
--- a/Engine/Plugins/Experimental/MeshModelingToolset/Source/MeshModelingToolsEditorOnly/Private/ShapeSprayTool.cpp
+++ b/Engine/Plugins/Experimental/MeshModelingToolset/Source/MeshModelingToolsEditorOnly/Private/ShapeSprayTool.cpp
@@ -3,9 +3,10 @@
 #include "ShapeSprayTool.h"
 #include "InteractiveToolManager.h"
 #include "ToolBuilderUtil.h"
+#include "ToolSetupUtil.h"
 #include "Generators/MinimalBoxMeshGenerator.h"
+#include "ModelingObjectsCreationAPI.h"
 #include "Selection/ToolSelectionUtil.h"
-#include "AssetGenerationUtil.h"
 #include "MeshNormals.h"
 
 #include "TargetInterfaces/PrimitiveComponentBackedTarget.h"
@@ -25,7 +26,6 @@ UMeshSurfacePointTool* UShapeSprayToolBuilder::CreateNewTool(const FToolBuilderS
 {
 	UShapeSprayTool* NewTool = NewObject<UShapeSprayTool>(SceneState.ToolManager);
 	NewTool->SetWorld(SceneState.World);
-	NewTool->SetAssetAPI(AssetAPI);
 	return NewTool;
 }
 
@@ -55,12 +55,6 @@ void UShapeSprayTool::SetWorld(UWorld* World)
 	this->TargetWorld = World;
 }
 
-void UShapeSprayTool::SetAssetAPI(IAssetGenerationAPI* AssetAPIIn)
-{
-	this->AssetAPI = AssetAPIIn;
-}
-
-
 
 void UShapeSprayTool::Setup()
 {
@@ -239,12 +233,16 @@ void UShapeSprayTool::EmitResult()
 
 	GetToolManager()->BeginUndoTransaction(LOCTEXT("EmitShapeSpray", "Create ShapeSpray"));
 
-	AActor* NewActor = AssetGenerationUtil::GenerateStaticMeshActor(
-		AssetAPI, TargetWorld,
-		Mesh, UseTransform, TEXT("Polygon"));
-	if (NewActor != nullptr)
+	FCreateMeshObjectParams NewMeshObjectParams;
+	NewMeshObjectParams.TargetWorld = TargetWorld;
+	NewMeshObjectParams.Transform = (FTransform)UseTransform;
+	NewMeshObjectParams.BaseName = TEXT("Spray");
+	NewMeshObjectParams.Materials.Add( ToolSetupUtil::GetDefaultMaterial() );
+	NewMeshObjectParams.SetMesh(Mesh);
+	FCreateMeshObjectResult Result = UE::Modeling::CreateMeshObject(GetToolManager(), MoveTemp(NewMeshObjectParams));
+	if (Result.IsOK() && Result.NewActor != nullptr)
 	{
-		ToolSelectionUtil::SetNewActorSelection(GetToolManager(), NewActor);
+		ToolSelectionUtil::SetNewActorSelection(GetToolManager(), Result.NewActor);
 	}
 
 	GetToolManager()->EndUndoTransaction();
diff --git a/Engine/Plugins/Experimental/MeshModelingToolset/Source/MeshModelingToolsEditorOnly/Private/SimplifyMeshTool.cpp b/Engine/Plugins/Experimental/MeshModelingToolset/Source/MeshModelingToolsEditorOnly/Private/SimplifyMeshTool.cpp
--- a/Engine/Plugins/Experimental/MeshModelingToolset/Source/MeshModelingToolsEditorOnly/Private/SimplifyMeshTool.cpp
+++ b/Engine/Plugins/Experimental/MeshModelingToolset/Source/MeshModelingToolsEditorOnly/Private/SimplifyMeshTool.cpp
@@ -14,8 +14,6 @@
 #include "DynamicMeshToMeshDescription.h"
 #include "AssetUtils/MeshDescriptionUtil.h"
 
-#include "AssetGenerationUtil.h"
-
 #include "SceneManagement.h" // for FPrimitiveDrawInterface
 
 #include "TargetInterfaces/MaterialProvider.h"
diff --git a/Engine/Plugins/Experimental/MeshModelingToolset/Source/MeshModelingToolsEditorOnly/Private/UVLayoutTool.cpp b/Engine/Plugins/Experimental/MeshModelingToolset/Source/MeshModelingToolsEditorOnly/Private/UVLayoutTool.cpp
--- a/Engine/Plugins/Experimental/MeshModelingToolset/Source/MeshModelingToolsEditorOnly/Private/UVLayoutTool.cpp
+++ b/Engine/Plugins/Experimental/MeshModelingToolset/Source/MeshModelingToolsEditorOnly/Private/UVLayoutTool.cpp
@@ -16,8 +16,6 @@
 
 #include "ParameterizationOps/UVLayoutOp.h"
 
-#include "AssetGenerationUtil.h"
-
 #include "TargetInterfaces/MaterialProvider.h"
 #include "TargetInterfaces/MeshDescriptionCommitter.h"
 #include "TargetInterfaces/MeshDescriptionProvider.h"
@@ -57,7 +55,6 @@ UInteractiveTool* UUVLayoutToolBuilder::BuildTool(const FToolBuilderState& Scene
 	TArray<TObjectPtr<UToolTarget>> Targets = SceneState.TargetManager->BuildAllSelectedTargetable(SceneState, GetTargetRequirements());
 	NewTool->SetTargets(MoveTemp(Targets));
 	NewTool->SetWorld(SceneState.World, SceneState.GizmoManager);
-	NewTool->SetAssetAPI(AssetAPI);
 
 	return NewTool;
 }
@@ -216,11 +213,6 @@ void UUVLayoutTool::Shutdown(EToolShutdownType ShutdownType)
 	}
 }
 
-void UUVLayoutTool::SetAssetAPI(IAssetGenerationAPI* AssetAPIIn)
-{
-	this->AssetAPI = AssetAPIIn;
-}
-
 TUniquePtr<FDynamicMeshOperator> UUVLayoutOperatorFactory::MakeNewOperator()
 {
 	TUniquePtr<FUVLayoutOp> Op = MakeUnique<FUVLayoutOp>();
diff --git a/Engine/Plugins/Experimental/MeshModelingToolset/Source/MeshModelingToolsEditorOnly/Private/VoxelCSGMeshesTool.cpp b/Engine/Plugins/Experimental/MeshModelingToolset/Source/MeshModelingToolsEditorOnly/Private/VoxelCSGMeshesTool.cpp
--- a/Engine/Plugins/Experimental/MeshModelingToolset/Source/MeshModelingToolsEditorOnly/Private/VoxelCSGMeshesTool.cpp
+++ b/Engine/Plugins/Experimental/MeshModelingToolset/Source/MeshModelingToolsEditorOnly/Private/VoxelCSGMeshesTool.cpp
@@ -4,7 +4,7 @@
 #include "InteractiveToolManager.h"
 #include "ToolBuilderUtil.h"
 #include "ToolSetupUtil.h"
-#include "AssetGenerationUtil.h"
+#include "ModelingObjectsCreationAPI.h"
 #include "Selection/ToolSelectionUtil.h"
 
 #include "DynamicMesh3.h"
@@ -38,9 +38,7 @@ const FToolTargetTypeRequirements& UVoxelCSGMeshesToolBuilder::GetTargetRequirem
 
 bool UVoxelCSGMeshesToolBuilder::CanBuildTool(const FToolBuilderState& SceneState) const
 {
-	bool bHasBuildAPI = (this->AssetAPI != nullptr);
-	bool bHasComponents = SceneState.TargetManager->CountSelectedAndTargetable(SceneState, GetTargetRequirements()) == 2;
-	return (bHasBuildAPI && bHasComponents);
+	return SceneState.TargetManager->CountSelectedAndTargetable(SceneState, GetTargetRequirements()) == 2;
 }
 
 UInteractiveTool* UVoxelCSGMeshesToolBuilder::BuildTool(const FToolBuilderState& SceneState) const
@@ -50,7 +48,6 @@ UInteractiveTool* UVoxelCSGMeshesToolBuilder::BuildTool(const FToolBuilderState&
 	TArray<TObjectPtr<UToolTarget>> Targets = SceneState.TargetManager->BuildAllSelectedTargetable(SceneState, GetTargetRequirements());
 	NewTool->SetTargets(MoveTemp(Targets));
 	NewTool->SetWorld(SceneState.World);
-	NewTool->SetAssetAPI(AssetAPI);
 
 	return NewTool;
 }
@@ -68,13 +65,6 @@ void UVoxelCSGMeshesTool::SetWorld(UWorld* World)
 	this->TargetWorld = World;
 }
 
-void UVoxelCSGMeshesTool::SetAssetAPI(IAssetGenerationAPI* AssetAPIIn)
-{
-	this->AssetAPI = AssetAPIIn;
-}
-
-
-
 void UVoxelCSGMeshesTool::Setup()
 {
 	UInteractiveTool::Setup();
@@ -221,16 +211,21 @@ void UVoxelCSGMeshesTool::CreateLowQualityPreview()
 	Preview->SetVisibility(true);
 }
 
-void UVoxelCSGMeshesTool::GenerateAsset(const FDynamicMeshOpResult& Result)
+void UVoxelCSGMeshesTool::GenerateAsset(const FDynamicMeshOpResult& OpResult)
 {
-	check(Result.Mesh.Get() != nullptr);
-
-	AActor* NewActor = AssetGenerationUtil::GenerateStaticMeshActor(
-		AssetAPI, TargetWorld,
-		Result.Mesh.Get(), Result.Transform, TEXT("CSGMesh"));
-	if (NewActor != nullptr)
+	if (ensure(OpResult.Mesh.Get() != nullptr))
 	{
-		ToolSelectionUtil::SetNewActorSelection(GetToolManager(), NewActor);
+		FCreateMeshObjectParams NewMeshObjectParams;
+		NewMeshObjectParams.TargetWorld = TargetWorld;
+		NewMeshObjectParams.Transform = (FTransform)OpResult.Transform;
+		NewMeshObjectParams.BaseName = TEXT("CSGMesh");
+		NewMeshObjectParams.Materials.Add(ToolSetupUtil::GetDefaultMaterial());
+		NewMeshObjectParams.SetMesh(OpResult.Mesh.Get());
+		FCreateMeshObjectResult Result = UE::Modeling::CreateMeshObject(GetToolManager(), MoveTemp(NewMeshObjectParams));
+		if (Result.IsOK() && Result.NewActor != nullptr)
+		{
+			ToolSelectionUtil::SetNewActorSelection(GetToolManager(), Result.NewActor);
+		}
 	}
 }
 
diff --git a/Engine/Plugins/Experimental/MeshModelingToolset/Source/MeshModelingToolsEditorOnly/Public/AttributeEditorTool.h b/Engine/Plugins/Experimental/MeshModelingToolset/Source/MeshModelingToolsEditorOnly/Public/AttributeEditorTool.h
--- a/Engine/Plugins/Experimental/MeshModelingToolset/Source/MeshModelingToolsEditorOnly/Public/AttributeEditorTool.h
+++ b/Engine/Plugins/Experimental/MeshModelingToolset/Source/MeshModelingToolsEditorOnly/Public/AttributeEditorTool.h
@@ -18,8 +18,6 @@ class MESHMODELINGTOOLSEDITORONLY_API UAttributeEditorToolBuilder : public UInte
 	GENERATED_BODY()
 
 public:
-	IToolsContextAssetAPI* AssetAPI = nullptr;
-
 	virtual bool CanBuildTool(const FToolBuilderState& SceneState) const override;
 	virtual UInteractiveTool* BuildTool(const FToolBuilderState& SceneState) const override;
 
diff --git a/Engine/Plugins/Experimental/MeshModelingToolset/Source/MeshModelingToolsEditorOnly/Public/BspConversionTool.h b/Engine/Plugins/Experimental/MeshModelingToolset/Source/MeshModelingToolsEditorOnly/Public/BspConversionTool.h
--- a/Engine/Plugins/Experimental/MeshModelingToolset/Source/MeshModelingToolsEditorOnly/Public/BspConversionTool.h
+++ b/Engine/Plugins/Experimental/MeshModelingToolset/Source/MeshModelingToolsEditorOnly/Public/BspConversionTool.h
@@ -14,7 +14,6 @@
 
 #include "BspConversionTool.generated.h"
 
-class IAssetGenerationAPI;
 
 /**
  * Builder for UBspConversionTool.
@@ -23,10 +22,7 @@ UCLASS()
 class MESHMODELINGTOOLSEDITORONLY_API UBspConversionToolBuilder : public UInteractiveToolBuilder
 {
 	GENERATED_BODY()
-
 public:
-	IAssetGenerationAPI* AssetAPI = nullptr;
-
 	virtual bool CanBuildTool(const FToolBuilderState& SceneState) const override;
 	virtual UInteractiveTool* BuildTool(const FToolBuilderState& SceneState) const override;
 };
@@ -141,7 +137,6 @@ public:
 	virtual void Shutdown(EToolShutdownType ShutdownType) override;
 
 	virtual void SetWorld(UWorld* World) { this->TargetWorld = World; }
-	virtual void SetAssetAPI(IAssetGenerationAPI* InAssetAPI) { this->AssetAPI = InAssetAPI; }
 
 	virtual void OnTick(float DeltaTime) override;
 	virtual void Render(IToolsContextRenderAPI* RenderAPI) override;
@@ -164,7 +159,6 @@ protected:
 	bool bCanAccept = false;
 
 	UWorld* TargetWorld;
-	IAssetGenerationAPI* AssetAPI;
 	EBspConversionToolAction PendingAction;
 
 	UPROPERTY()
diff --git a/Engine/Plugins/Experimental/MeshModelingToolset/Source/MeshModelingToolsEditorOnly/Public/GenerateLODMeshesTool.h b/Engine/Plugins/Experimental/MeshModelingToolset/Source/MeshModelingToolsEditorOnly/Public/GenerateLODMeshesTool.h
--- a/Engine/Plugins/Experimental/MeshModelingToolset/Source/MeshModelingToolsEditorOnly/Public/GenerateLODMeshesTool.h
+++ b/Engine/Plugins/Experimental/MeshModelingToolset/Source/MeshModelingToolsEditorOnly/Public/GenerateLODMeshesTool.h
@@ -15,10 +15,6 @@
 #include "BaseTools/SingleSelectionMeshEditingTool.h"
 #include "GenerateLODMeshesTool.generated.h"
 
-
-class IAssetGenerationAPI;
-
-
 /**
  *
  */
@@ -26,10 +22,8 @@ UCLASS()
 class MESHMODELINGTOOLSEDITORONLY_API UGenerateLODMeshesToolBuilder : public USingleSelectionMeshEditingToolBuilder
 {
 	GENERATED_BODY()
-
 public:
 	virtual USingleSelectionMeshEditingTool* CreateNewTool(const FToolBuilderState& SceneState) const override;
-	virtual bool RequiresAssetAPI() const override { return true; }
 };
 
 
diff --git a/Engine/Plugins/Experimental/MeshModelingToolset/Source/MeshModelingToolsEditorOnly/Public/Hair/GroomCardsEditorTool.h b/Engine/Plugins/Experimental/MeshModelingToolset/Source/MeshModelingToolsEditorOnly/Public/Hair/GroomCardsEditorTool.h
--- a/Engine/Plugins/Experimental/MeshModelingToolset/Source/MeshModelingToolsEditorOnly/Public/Hair/GroomCardsEditorTool.h
+++ b/Engine/Plugins/Experimental/MeshModelingToolset/Source/MeshModelingToolsEditorOnly/Public/Hair/GroomCardsEditorTool.h
@@ -33,8 +33,6 @@ class MESHMODELINGTOOLSEDITORONLY_API UGroomCardsEditorToolBuilder : public UMes
 {
 	GENERATED_BODY()
 public:
-	IToolsContextAssetAPI* AssetAPI = nullptr;
-
 	virtual UMeshSurfacePointTool* CreateNewTool(const FToolBuilderState& SceneState) const override;
 };
 
@@ -151,7 +149,6 @@ public:
 	UGroomCardsEditorTool();
 
 	virtual void SetWorld(UWorld* World) { this->TargetWorld = World; }
-	virtual void SetAssetAPI(IToolsContextAssetAPI* InAssetAPI) { this->AssetAPI = InAssetAPI; }
 
 	virtual void Setup() override;
 	virtual void Shutdown(EToolShutdownType ShutdownType) override;
@@ -225,7 +222,6 @@ protected:		// mechanics
 
 protected:
 	UWorld* TargetWorld = nullptr;
-	IToolsContextAssetAPI* AssetAPI = nullptr;
 
 	TPimplPtr<FEditableGroomCardSet> EditableCardSet;
 
diff --git a/Engine/Plugins/Experimental/MeshModelingToolset/Source/MeshModelingToolsEditorOnly/Public/Hair/GroomToMeshTool.h b/Engine/Plugins/Experimental/MeshModelingToolset/Source/MeshModelingToolsEditorOnly/Public/Hair/GroomToMeshTool.h
--- a/Engine/Plugins/Experimental/MeshModelingToolset/Source/MeshModelingToolsEditorOnly/Public/Hair/GroomToMeshTool.h
+++ b/Engine/Plugins/Experimental/MeshModelingToolset/Source/MeshModelingToolsEditorOnly/Public/Hair/GroomToMeshTool.h
@@ -13,7 +13,6 @@
 
 class AGroomActor;
 class AStaticMeshActor;
-class IAssetGenerationAPI;
 PREDECLARE_USE_GEOMETRY_CLASS(FDynamicMesh3);
 
 /**
@@ -24,8 +23,6 @@ class MESHMODELINGTOOLSEDITORONLY_API UGroomToMeshToolBuilder : public UInteract
 {
 	GENERATED_BODY()
 public:
-	IAssetGenerationAPI* AssetAPI = nullptr;
-
 	virtual bool CanBuildTool(const FToolBuilderState& SceneState) const override;
 	virtual UInteractiveTool* BuildTool(const FToolBuilderState& SceneState) const override;
 };
@@ -126,7 +123,6 @@ public:
 	UGroomToMeshTool();
 
 	virtual void SetWorld(UWorld* World) { this->TargetWorld = World; }
-	virtual void SetAssetAPI(IAssetGenerationAPI* InAssetAPI) { this->AssetAPI = InAssetAPI; }
 	virtual void SetSelection(AGroomActor* Groom);
 
 	virtual void Setup() override;
@@ -161,7 +157,6 @@ protected:
 
 protected:
 	UWorld* TargetWorld = nullptr;
-	IAssetGenerationAPI* AssetAPI = nullptr;
 
 	FDynamicMesh3 CurrentMesh;
 
diff --git a/Engine/Plugins/Experimental/MeshModelingToolset/Source/MeshModelingToolsEditorOnly/Public/MergeMeshesTool.h b/Engine/Plugins/Experimental/MeshModelingToolset/Source/MeshModelingToolsEditorOnly/Public/MergeMeshesTool.h
--- a/Engine/Plugins/Experimental/MeshModelingToolset/Source/MeshModelingToolsEditorOnly/Public/MergeMeshesTool.h
+++ b/Engine/Plugins/Experimental/MeshModelingToolset/Source/MeshModelingToolsEditorOnly/Public/MergeMeshesTool.h
@@ -11,9 +11,6 @@
 #include "PropertySets/OnAcceptProperties.h"
 #include "MergeMeshesTool.generated.h"
 
-class IAssetGenerationAPI;
-
-
 /**
  *
  */
@@ -23,13 +20,6 @@ class MESHMODELINGTOOLSEDITORONLY_API UMergeMeshesToolBuilder : public UInteract
 	GENERATED_BODY()
 
 public:
-	IAssetGenerationAPI* AssetAPI;
-
-	UMergeMeshesToolBuilder()
-	{
-		AssetAPI = nullptr;
-	}
-
 	virtual bool CanBuildTool(const FToolBuilderState& SceneState) const override;
 	virtual UInteractiveTool* BuildTool(const FToolBuilderState& SceneState) const override;
 
@@ -77,7 +67,6 @@ public:
 	UMergeMeshesTool();
 
 	virtual void SetWorld(UWorld* World);
-	virtual void SetAssetAPI(IAssetGenerationAPI* AssetAPI);
 
 	virtual void Setup() override;
 	virtual void Shutdown(EToolShutdownType ShutdownType) override;
@@ -109,7 +98,6 @@ protected:
 
 protected:
 	UWorld* TargetWorld;
-	IAssetGenerationAPI* AssetAPI;
 
 	TSharedPtr<TArray<IVoxelBasedCSG::FPlacedMesh>> InputMeshes;
 	/** stash copies of the transforms and pointers to the meshes for consumption by merge Op*/
diff --git a/Engine/Plugins/Experimental/MeshModelingToolset/Source/MeshModelingToolsEditorOnly/Public/MeshToVolumeTool.h b/Engine/Plugins/Experimental/MeshModelingToolset/Source/MeshModelingToolsEditorOnly/Public/MeshToVolumeTool.h
--- a/Engine/Plugins/Experimental/MeshModelingToolset/Source/MeshModelingToolsEditorOnly/Public/MeshToVolumeTool.h
+++ b/Engine/Plugins/Experimental/MeshModelingToolset/Source/MeshModelingToolsEditorOnly/Public/MeshToVolumeTool.h
@@ -16,8 +16,6 @@
 #include "BaseTools/SingleSelectionMeshEditingTool.h"
 #include "MeshToVolumeTool.generated.h"
 
-class IAssetGenerationAPI;
-
 /**
  *
  */
diff --git a/Engine/Plugins/Experimental/MeshModelingToolset/Source/MeshModelingToolsEditorOnly/Public/ShapeSprayTool.h b/Engine/Plugins/Experimental/MeshModelingToolset/Source/MeshModelingToolsEditorOnly/Public/ShapeSprayTool.h
--- a/Engine/Plugins/Experimental/MeshModelingToolset/Source/MeshModelingToolsEditorOnly/Public/ShapeSprayTool.h
+++ b/Engine/Plugins/Experimental/MeshModelingToolset/Source/MeshModelingToolsEditorOnly/Public/ShapeSprayTool.h
@@ -3,13 +3,10 @@
 #pragma once
 
 #include "CoreMinimal.h"
-#include "UObject/NoExportTypes.h"
 #include "InteractiveToolBuilder.h"
 #include "DynamicMeshBrushTool.h"
 #include "ShapeSprayTool.generated.h"
 
-class IAssetGenerationAPI;
-
 /**
  * UMeshSurfacePointToolBuilder override for UShapeSprayTool
  */
@@ -19,8 +16,6 @@ class MESHMODELINGTOOLSEDITORONLY_API UShapeSprayToolBuilder : public UMeshSurfa
 	GENERATED_BODY()
 
 public:
-	IAssetGenerationAPI* AssetAPI = nullptr;
-
 	virtual UMeshSurfacePointTool* CreateNewTool(const FToolBuilderState& SceneState) const override;
 };
 
@@ -71,7 +66,6 @@ public:
 	UShapeSprayTool();
 
 	virtual void SetWorld(UWorld* World);
-	virtual void SetAssetAPI(IAssetGenerationAPI* AssetAPI);
 
 	// UInteractiveTool API
 
@@ -100,7 +94,6 @@ protected:
 
 protected:
 	UWorld* TargetWorld;
-	IAssetGenerationAPI* AssetAPI;
 
 	UE::Geometry::FDynamicMesh3 ShapeMesh;
 	void UpdateShapeMesh();
diff --git a/Engine/Plugins/Experimental/MeshModelingToolset/Source/MeshModelingToolsEditorOnly/Public/UVLayoutTool.h b/Engine/Plugins/Experimental/MeshModelingToolset/Source/MeshModelingToolsEditorOnly/Public/UVLayoutTool.h
--- a/Engine/Plugins/Experimental/MeshModelingToolset/Source/MeshModelingToolsEditorOnly/Public/UVLayoutTool.h
+++ b/Engine/Plugins/Experimental/MeshModelingToolset/Source/MeshModelingToolsEditorOnly/Public/UVLayoutTool.h
@@ -20,8 +20,6 @@
 // predeclarations
 struct FMeshDescription;
 class USimpleDynamicMeshComponent;
-class IAssetGenerationAPI;
-
 
 /**
  *
@@ -32,8 +30,6 @@ class MESHMODELINGTOOLSEDITORONLY_API UUVLayoutToolBuilder : public UInteractive
 	GENERATED_BODY()
 
 public:
-	IAssetGenerationAPI* AssetAPI = nullptr;
-
 	virtual bool CanBuildTool(const FToolBuilderState& SceneState) const override;
 	virtual UInteractiveTool* BuildTool(const FToolBuilderState& SceneState) const override;
 
@@ -129,7 +125,6 @@ public:
 	virtual void Shutdown(EToolShutdownType ShutdownType) override;
 
 	virtual void SetWorld(UWorld* World, UInteractiveGizmoManager* GizmoManagerIn);
-	virtual void SetAssetAPI(IAssetGenerationAPI* AssetAPI);
 
 	virtual void OnTick(float DeltaTime) override;
 	virtual void Render(IToolsContextRenderAPI* RenderAPI) override;
@@ -163,7 +158,6 @@ protected:
 	TArray<TSharedPtr<UE::Geometry::FDynamicMesh3, ESPMode::ThreadSafe>> OriginalDynamicMeshes;
 
 	UWorld* TargetWorld = nullptr;
-	IAssetGenerationAPI* AssetAPI = nullptr;
 
 	FViewCameraState CameraState;
 
diff --git a/Engine/Plugins/Experimental/MeshModelingToolset/Source/MeshModelingToolsEditorOnly/Public/VoxelCSGMeshesTool.h b/Engine/Plugins/Experimental/MeshModelingToolset/Source/MeshModelingToolsEditorOnly/Public/VoxelCSGMeshesTool.h
--- a/Engine/Plugins/Experimental/MeshModelingToolset/Source/MeshModelingToolsEditorOnly/Public/VoxelCSGMeshesTool.h
+++ b/Engine/Plugins/Experimental/MeshModelingToolset/Source/MeshModelingToolsEditorOnly/Public/VoxelCSGMeshesTool.h
@@ -13,9 +13,6 @@
 
 #include "VoxelCSGMeshesTool.generated.h"
 
-class IAssetGenerationAPI;
-
-
 /**
  *
  */
@@ -25,13 +22,6 @@ class MESHMODELINGTOOLSEDITORONLY_API UVoxelCSGMeshesToolBuilder : public UInter
 	GENERATED_BODY()
 
 public:
-	IAssetGenerationAPI* AssetAPI;
-
-	UVoxelCSGMeshesToolBuilder()
-	{
-		AssetAPI = nullptr;
-	}
-
 	virtual bool CanBuildTool(const FToolBuilderState& SceneState) const override;
 	virtual UInteractiveTool* BuildTool(const FToolBuilderState& SceneState) const override;
 
@@ -103,7 +93,6 @@ public:
 	UVoxelCSGMeshesTool();
 
 	virtual void SetWorld(UWorld* World);
-	virtual void SetAssetAPI(IAssetGenerationAPI* AssetAPI);
 
 	virtual void Setup() override;
 	virtual void Shutdown(EToolShutdownType ShutdownType) override;
@@ -136,7 +125,6 @@ protected:
 
 protected:
 	UWorld* TargetWorld;
-	IAssetGenerationAPI* AssetAPI;
 
 	TSharedPtr<TArray<IVoxelBasedCSG::FPlacedMesh>> InputMeshes;
 	/** stash copies of the transforms and pointers to the meshes for consumption by the CSG Op*/
diff --git a/Engine/Plugins/Experimental/MeshModelingToolset/Source/ModelingComponents/ModelingComponents.Build.cs b/Engine/Plugins/Experimental/MeshModelingToolset/Source/ModelingComponents/ModelingComponents.Build.cs
--- a/Engine/Plugins/Experimental/MeshModelingToolset/Source/ModelingComponents/ModelingComponents.Build.cs
+++ b/Engine/Plugins/Experimental/MeshModelingToolset/Source/ModelingComponents/ModelingComponents.Build.cs
@@ -26,7 +26,8 @@ public class ModelingComponents : ModuleRules
 			new string[]
 			{
 				"Core",
-                "InteractiveToolsFramework",
+				"PhysicsCore",
+				"InteractiveToolsFramework",
                 "MeshDescription",
 				"StaticMeshDescription",
 				"GeometryCore",
diff --git a/Engine/Plugins/Experimental/MeshModelingToolset/Source/ModelingComponents/Private/AssetGenerationUtil.cpp b/Engine/Plugins/Experimental/MeshModelingToolset/Source/ModelingComponents/Private/AssetGenerationUtil.cpp
deleted file mode 100644
diff --git a/Engine/Plugins/Experimental/MeshModelingToolset/Source/ModelingComponents/Private/BaseTools/BaseCreateFromSelectedTool.cpp b/Engine/Plugins/Experimental/MeshModelingToolset/Source/ModelingComponents/Private/BaseTools/BaseCreateFromSelectedTool.cpp
--- a/Engine/Plugins/Experimental/MeshModelingToolset/Source/ModelingComponents/Private/BaseTools/BaseCreateFromSelectedTool.cpp
+++ b/Engine/Plugins/Experimental/MeshModelingToolset/Source/ModelingComponents/Private/BaseTools/BaseCreateFromSelectedTool.cpp
@@ -13,7 +13,7 @@
 #include "DynamicMeshToMeshDescription.h"
 #include "MeshDescriptionToDynamicMesh.h"
 
-#include "AssetGenerationUtil.h"
+#include "ModelingObjectsCreationAPI.h"
 #include "Selection/ToolSelectionUtil.h"
 
 #include "TargetInterfaces/MeshDescriptionCommitter.h"
@@ -46,7 +46,7 @@ const FToolTargetTypeRequirements& UBaseCreateFromSelectedToolBuilder::GetTarget
 bool UBaseCreateFromSelectedToolBuilder::CanBuildTool(const FToolBuilderState& SceneState) const
 {
 	int32 ComponentCount = SceneState.TargetManager->CountSelectedAndTargetable(SceneState, GetTargetRequirements());
-	return AssetAPI != nullptr && ComponentCount >= MinComponentsSupported() && (!MaxComponentsSupported().IsSet() || ComponentCount <= MaxComponentsSupported().GetValue());
+	return ComponentCount >= MinComponentsSupported() && (!MaxComponentsSupported().IsSet() || ComponentCount <= MaxComponentsSupported().GetValue());
 }
 
 
@@ -57,7 +57,6 @@ UInteractiveTool* UBaseCreateFromSelectedToolBuilder::BuildTool(const FToolBuild
 	TArray<TObjectPtr<UToolTarget>> Targets = SceneState.TargetManager->BuildAllSelectedTargetable(SceneState, GetTargetRequirements());
 	NewTool->SetTargets(MoveTemp(Targets));
 	NewTool->SetWorld(SceneState.World);
-	NewTool->SetAssetAPI(AssetAPI);
 
 	return NewTool;
 }
@@ -117,7 +116,7 @@ void UBaseCreateFromSelectedTool::Setup()
 		else
 		{
 			int32 Index = (HandleSourcesProperties->WriteOutputTo == EBaseCreateFromSelectedTargetType::FirstInputAsset) ? 0 : Targets.Num() - 1;
-			HandleSourcesProperties->OutputAsset = AssetGenerationUtil::GetComponentAssetBaseName(TargetComponentInterface(Index)->GetOwnerComponent(), false);
+			HandleSourcesProperties->OutputAsset = UE::Modeling::GetComponentAssetBaseName(TargetComponentInterface(Index)->GetOwnerComponent(), false);
 
 			// Reset the hidden gizmo to its initial position
 			FTransform ComponentTransform = TargetComponentInterface(Index)->GetWorldTransform();
@@ -149,12 +148,6 @@ void UBaseCreateFromSelectedTool::SetWorld(UWorld* World)
 }
 
 
-void UBaseCreateFromSelectedTool::SetAssetAPI(IAssetGenerationAPI* AssetAPIIn)
-{
-	this->AssetAPI = AssetAPIIn;
-}
-
-
 void UBaseCreateFromSelectedTool::OnTick(float DeltaTime)
 {
 	for (UTransformGizmo* Gizmo : TransformGizmos)
@@ -212,25 +205,25 @@ TArray<UMaterialInterface*> UBaseCreateFromSelectedTool::GetOutputMaterials() co
 }
 
 
-void UBaseCreateFromSelectedTool::GenerateAsset(const FDynamicMeshOpResult& Result)
+void UBaseCreateFromSelectedTool::GenerateAsset(const FDynamicMeshOpResult& OpResult)
 {
-	check(Result.Mesh.Get() != nullptr);
+	if (OpResult.Mesh.Get() == nullptr) return;
 
 	FTransform3d NewTransform;
 	if (Targets.Num() == 1) // in the single-selection case, shove the result back into the original component space
 	{
 		FTransform3d ToSourceComponentSpace = (FTransform3d)TargetComponentInterface(0)->GetWorldTransform().Inverse();
-		MeshTransforms::ApplyTransform(*Result.Mesh, ToSourceComponentSpace);
+		MeshTransforms::ApplyTransform(*OpResult.Mesh, ToSourceComponentSpace);
 		NewTransform = (FTransform3d)TargetComponentInterface(0)->GetWorldTransform();
 	}
 	else // in the multi-selection case, center the pivot for the combined result
 	{
-		FVector3d Center = Result.Mesh->GetCachedBounds().Center();
-		double Rescale = Result.Transform.GetScale().X;
+		FVector3d Center = OpResult.Mesh->GetCachedBounds().Center();
+		double Rescale = OpResult.Transform.GetScale().X;
 		FTransform3d LocalTransform(-Center * Rescale);
 		LocalTransform.SetScale(FVector3d(Rescale, Rescale, Rescale));
-		MeshTransforms::ApplyTransform(*Result.Mesh, LocalTransform);
-		NewTransform = Result.Transform;
+		MeshTransforms::ApplyTransform(*OpResult.Mesh, LocalTransform);
+		NewTransform = OpResult.Transform;
 		NewTransform.SetScale(FVector3d::One());
 		NewTransform.SetTranslation(NewTransform.GetTranslation() + NewTransform.TransformVector(Center * Rescale));
 	}
@@ -242,13 +235,16 @@ void UBaseCreateFromSelectedTool::GenerateAsset(const FDynamicMeshOpResult& Resu
 		UseBaseName = PrefixWithSourceNameIfSingleSelection(GetCreatedAssetName());
 	}
 
-	TArray<UMaterialInterface*> Materials = GetOutputMaterials();
-	AActor* NewActor = AssetGenerationUtil::GenerateStaticMeshActor(
-		AssetAPI, TargetWorld,
-		Result.Mesh.Get(), NewTransform, UseBaseName, Materials);
-	if (NewActor != nullptr)
+	FCreateMeshObjectParams NewMeshObjectParams;
+	NewMeshObjectParams.TargetWorld = TargetWorld;
+	NewMeshObjectParams.Transform = (FTransform)NewTransform;
+	NewMeshObjectParams.BaseName = UseBaseName;
+	NewMeshObjectParams.Materials = GetOutputMaterials();
+	NewMeshObjectParams.SetMesh(OpResult.Mesh.Get());
+	FCreateMeshObjectResult Result = UE::Modeling::CreateMeshObject(GetToolManager(), MoveTemp(NewMeshObjectParams));
+	if (Result.IsOK() && Result.NewActor != nullptr)
 	{
-		ToolSelectionUtil::SetNewActorSelection(GetToolManager(), NewActor);
+		ToolSelectionUtil::SetNewActorSelection(GetToolManager(), Result.NewActor);
 	}
 }
 
@@ -284,7 +280,7 @@ FString UBaseCreateFromSelectedTool::PrefixWithSourceNameIfSingleSelection(const
 {
 	if (Targets.Num() == 1)
 	{
-		FString CurName = AssetGenerationUtil::GetComponentAssetBaseName(TargetComponentInterface(0)->GetOwnerComponent());
+		FString CurName = UE::Modeling::GetComponentAssetBaseName(TargetComponentInterface(0)->GetOwnerComponent());
 		return FString::Printf(TEXT("%s_%s"), *CurName, *AssetName);
 	}
 	else
diff --git a/Engine/Plugins/Experimental/MeshModelingToolset/Source/ModelingComponents/Private/BaseTools/SingleSelectionMeshEditingTool.cpp b/Engine/Plugins/Experimental/MeshModelingToolset/Source/ModelingComponents/Private/BaseTools/SingleSelectionMeshEditingTool.cpp
--- a/Engine/Plugins/Experimental/MeshModelingToolset/Source/ModelingComponents/Private/BaseTools/SingleSelectionMeshEditingTool.cpp
+++ b/Engine/Plugins/Experimental/MeshModelingToolset/Source/ModelingComponents/Private/BaseTools/SingleSelectionMeshEditingTool.cpp
@@ -11,7 +11,6 @@
 /*
  * ToolBuilder
  */
-
 const FToolTargetTypeRequirements& USingleSelectionMeshEditingToolBuilder::GetTargetRequirements() const
 {
 	static FToolTargetTypeRequirements TypeRequirements({
@@ -25,8 +24,7 @@ const FToolTargetTypeRequirements& USingleSelectionMeshEditingToolBuilder::GetTa
 
 bool USingleSelectionMeshEditingToolBuilder::CanBuildTool(const FToolBuilderState& SceneState) const
 {
-	bool bValidAssetAPI = (RequiresAssetAPI() ? AssetAPI != nullptr : true);
-	return bValidAssetAPI && SceneState.TargetManager->CountSelectedAndTargetable(SceneState, GetTargetRequirements()) == 1;
+	return SceneState.TargetManager->CountSelectedAndTargetable(SceneState, GetTargetRequirements()) == 1;
 }
 
 UInteractiveTool* USingleSelectionMeshEditingToolBuilder::BuildTool(const FToolBuilderState& SceneState) const
@@ -42,10 +40,6 @@ void USingleSelectionMeshEditingToolBuilder::InitializeNewTool(USingleSelectionM
 	check(Target);
 	NewTool->SetTarget(Target);
 	NewTool->SetWorld(SceneState.World);
-	if (RequiresAssetAPI())
-	{
-		NewTool->SetAssetAPI(AssetAPI);
-	}
 
 	if (WantsInputSelectionIfAvailable() && SceneState.StoredToolSelection != nullptr)
 	{
diff --git a/Engine/Plugins/Experimental/MeshModelingToolset/Source/ModelingComponents/Private/ModelingObjectsCreationAPI.cpp b/Engine/Plugins/Experimental/MeshModelingToolset/Source/ModelingComponents/Private/ModelingObjectsCreationAPI.cpp
new file mode 100644
--- /dev/null
+++ b/Engine/Plugins/Experimental/MeshModelingToolset/Source/ModelingComponents/Private/ModelingObjectsCreationAPI.cpp
@@ -0,0 +1,152 @@
+// Copyright Epic Games, Inc. All Rights Reserved.
+
+#include "ModelingObjectsCreationAPI.h"
+#include "InteractiveToolManager.h"
+#include "ContextObjectStore.h"
+
+#include "Misc/Paths.h"
+#include "Components/StaticMeshComponent.h"
+#include "Engine/StaticMesh.h"
+
+
+
+void FCreateMeshObjectParams::SetMesh(FMeshDescription&& MeshDescriptionIn)
+{
+	MeshDescription = MoveTemp(MeshDescriptionIn);
+	MeshType = ECreateMeshObjectSourceMeshType::MeshDescription;
+}
+
+void FCreateMeshObjectParams::SetMesh(const UE::Geometry::FDynamicMesh3* DynamicMeshIn)
+{
+	DynamicMesh = *DynamicMeshIn;
+	MeshType = ECreateMeshObjectSourceMeshType::DynamicMesh;
+}
+
+void FCreateMeshObjectParams::SetMesh(UE::Geometry::FDynamicMesh3&& DynamicMeshIn)
+{
+	DynamicMesh = MoveTemp(DynamicMeshIn);
+	MeshType = ECreateMeshObjectSourceMeshType::DynamicMesh;
+}
+
+
+
+
+
+
+
+FCreateMeshObjectResult UE::Modeling::CreateMeshObject(UInteractiveToolManager* ToolManager, FCreateMeshObjectParams&& CreateMeshParams)
+{
+	if (ensure(ToolManager))
+	{
+		UModelingObjectsCreationAPI* UseAPI = ToolManager->GetContextObjectStore()->FindContext<UModelingObjectsCreationAPI>();
+		if (UseAPI)
+		{
+			if (UseAPI->HasMoveVariants())
+			{
+				return UseAPI->CreateMeshObject(MoveTemp(CreateMeshParams));
+			}
+			else
+			{
+				return UseAPI->CreateMeshObject(CreateMeshParams);
+			}
+		}
+	}
+	return FCreateMeshObjectResult{ ECreateModelingObjectResult::Failed_NoAPIFound };
+}
+
+
+
+
+
+FCreateTextureObjectResult UE::Modeling::CreateTextureObject(UInteractiveToolManager* ToolManager, FCreateTextureObjectParams&& CreateTexParams)
+{
+	if (ensure(ToolManager))
+	{
+		UModelingObjectsCreationAPI* UseAPI = ToolManager->GetContextObjectStore()->FindContext<UModelingObjectsCreationAPI>();
+		if (UseAPI)
+		{
+			if (UseAPI->HasMoveVariants())
+			{
+				return UseAPI->CreateTextureObject(MoveTemp(CreateTexParams));
+			}
+			else
+			{
+				return UseAPI->CreateTextureObject(CreateTexParams);
+			}
+		}
+	}
+	return FCreateTextureObjectResult{ ECreateModelingObjectResult::Failed_NoAPIFound };
+}
+
+
+
+
+
+
+
+FString UE::Modeling::GetComponentAssetBaseName(UPrimitiveComponent* Component, bool bRemoveAutoGeneratedSuffixes)
+{
+	if (!ensure(Component != nullptr))
+	{
+		return TEXT("InvalidComponent");
+	}
+
+	FString ResultName = Component->GetName();
+	if (bRemoveAutoGeneratedSuffixes)
+	{
+		ResultName = UE::Modeling::StripGeneratedAssetSuffixFromName(ResultName);
+	}
+
+	// handle static meshcomponents
+	UStaticMeshComponent* StaticMeshComponent = Cast<UStaticMeshComponent>(Component);
+	if (StaticMeshComponent != nullptr)
+	{
+		UStaticMesh* SourceMesh = StaticMeshComponent->GetStaticMesh();
+		if (SourceMesh)
+		{
+			FString AssetName = FPaths::GetBaseFilename(SourceMesh->GetName());
+			ResultName = (bRemoveAutoGeneratedSuffixes) ? UE::Modeling::StripGeneratedAssetSuffixFromName(AssetName) : AssetName;
+		}
+	}
+
+	return ResultName;
+}
+
+
+
+
+FString UE::Modeling::StripGeneratedAssetSuffixFromName(FString InputName)
+{
+	// find final '_'
+	int32 Index;
+	if (!InputName.FindLastChar('_', Index))
+	{
+		return InputName;
+	}
+	// check that remaining characters are hex digits (from UUID)
+	int32 Len = InputName.Len();
+	int32 Count = 0, Letters = 0, Numbers = 0;
+	for (int32 k = Index + 1; k < Len; ++k)
+	{
+		if (FChar::IsHexDigit(InputName[k]) == false)
+		{
+			return InputName;
+		}
+		Count++;
+		if (FChar::IsDigit(InputName[k]))
+		{
+			Numbers++;
+		}
+		else
+		{
+			Letters++;
+		}
+	}
+	// currently assuming appended UUID is at least 8 characters
+	if (Numbers == 0 || Letters == 0 || Count < 8)
+	{
+		return InputName;
+	}
+
+	return InputName.Left(Index);
+}
diff --git a/Engine/Plugins/Experimental/MeshModelingToolset/Source/ModelingComponents/Private/Physics/PhysicsDataCollection.cpp b/Engine/Plugins/Experimental/MeshModelingToolset/Source/ModelingComponents/Private/Physics/PhysicsDataCollection.cpp
--- a/Engine/Plugins/Experimental/MeshModelingToolset/Source/ModelingComponents/Private/Physics/PhysicsDataCollection.cpp
+++ b/Engine/Plugins/Experimental/MeshModelingToolset/Source/ModelingComponents/Private/Physics/PhysicsDataCollection.cpp
@@ -11,7 +11,7 @@ using namespace UE::Geometry;
 
 void FPhysicsDataCollection::InitializeFromComponent(const UActorComponent* Component, bool bInitializeAggGeom)
 {
-	const UStaticMeshComponent* StaticMeshComponent = CastChecked<UStaticMeshComponent>(Component);
+	const UStaticMeshComponent* StaticMeshComponent = Cast<UStaticMeshComponent>(Component);
 	if (ensure(StaticMeshComponent))
 	{
 		const UStaticMesh* StaticMesh = StaticMeshComponent->GetStaticMesh();
diff --git a/Engine/Plugins/Experimental/MeshModelingToolset/Source/ModelingComponents/Public/AssetGenerationUtil.h b/Engine/Plugins/Experimental/MeshModelingToolset/Source/ModelingComponents/Public/AssetGenerationUtil.h
deleted file mode 100644
diff --git a/Engine/Plugins/Experimental/MeshModelingToolset/Source/ModelingComponents/Public/BaseTools/BaseCreateFromSelectedTool.h b/Engine/Plugins/Experimental/MeshModelingToolset/Source/ModelingComponents/Public/BaseTools/BaseCreateFromSelectedTool.h
--- a/Engine/Plugins/Experimental/MeshModelingToolset/Source/ModelingComponents/Public/BaseTools/BaseCreateFromSelectedTool.h
+++ b/Engine/Plugins/Experimental/MeshModelingToolset/Source/ModelingComponents/Public/BaseTools/BaseCreateFromSelectedTool.h
@@ -9,7 +9,6 @@
 #include "PropertySets/OnAcceptProperties.h"
 #include "BaseCreateFromSelectedTool.generated.h"
 
-class IAssetGenerationAPI;
 class UTransformGizmo;
 class UTransformProxy;
 
@@ -26,8 +25,6 @@ public:
 	virtual UInteractiveTool* BuildTool(const FToolBuilderState& SceneState) const override;
 
 public:
-	IAssetGenerationAPI* AssetAPI = nullptr;
-	
 	virtual TOptional<int32> MaxComponentsSupported() const { return TOptional<int32>(); }
 	virtual int32 MinComponentsSupported() const { return 1; }
 	
@@ -108,7 +105,6 @@ public:
 	UBaseCreateFromSelectedTool() = default;
 
 	virtual void SetWorld(UWorld* World);
-	virtual void SetAssetAPI(IAssetGenerationAPI* AssetAPI);
 
 	//
 	// InteractiveTool API - generally does not need to be modified by subclasses
@@ -207,8 +203,6 @@ protected:
 	UPROPERTY()
 	TArray<UTransformGizmo*> TransformGizmos;
 
-
 	UWorld* TargetWorld = nullptr;
-	IAssetGenerationAPI* AssetAPI = nullptr;
 };
 
diff --git a/Engine/Plugins/Experimental/MeshModelingToolset/Source/ModelingComponents/Public/BaseTools/BaseMeshProcessingTool.h b/Engine/Plugins/Experimental/MeshModelingToolset/Source/ModelingComponents/Public/BaseTools/BaseMeshProcessingTool.h
--- a/Engine/Plugins/Experimental/MeshModelingToolset/Source/ModelingComponents/Public/BaseTools/BaseMeshProcessingTool.h
+++ b/Engine/Plugins/Experimental/MeshModelingToolset/Source/ModelingComponents/Public/BaseTools/BaseMeshProcessingTool.h
@@ -194,7 +194,6 @@ protected:
 
 protected:
 	UWorld* TargetWorld = nullptr;
-	IToolsContextAssetAPI* AssetAPI = nullptr;
 
 	// Preview object holds temporary Actor with preview mesh component
 	UPROPERTY()
diff --git a/Engine/Plugins/Experimental/MeshModelingToolset/Source/ModelingComponents/Public/BaseTools/SingleSelectionMeshEditingTool.h b/Engine/Plugins/Experimental/MeshModelingToolset/Source/ModelingComponents/Public/BaseTools/SingleSelectionMeshEditingTool.h
--- a/Engine/Plugins/Experimental/MeshModelingToolset/Source/ModelingComponents/Public/BaseTools/SingleSelectionMeshEditingTool.h
+++ b/Engine/Plugins/Experimental/MeshModelingToolset/Source/ModelingComponents/Public/BaseTools/SingleSelectionMeshEditingTool.h
@@ -7,7 +7,6 @@
 #include "InteractiveToolStorableSelection.h"
 #include "SingleSelectionMeshEditingTool.generated.h"
 
-class IAssetGenerationAPI;
 class USingleSelectionMeshEditingTool;
 
 /**
@@ -20,9 +19,6 @@ class MODELINGCOMPONENTS_API USingleSelectionMeshEditingToolBuilder : public UIn
 {
 	GENERATED_BODY()
 
-public:
-	IAssetGenerationAPI* AssetAPI = nullptr;
-
 public:
 	/** @return true if a single mesh source can be found in the active selection */
 	virtual bool CanBuildTool(const FToolBuilderState& SceneState) const override;
@@ -36,9 +32,6 @@ public:
 	/** Called by BuildTool to configure the Tool with the input MeshSource based on the SceneState */
 	virtual void InitializeNewTool(USingleSelectionMeshEditingTool* Tool, const FToolBuilderState& SceneState) const;
 
-	/** @return true if this Tool requires an AssetAPI */
-	virtual bool RequiresAssetAPI() const { return false; }
-
 	/** @return true if this Tool would like access to an available Input Selection object */
 	virtual bool WantsInputSelectionIfAvailable() const { return false; }
 
@@ -57,13 +50,9 @@ class MODELINGCOMPONENTS_API USingleSelectionMeshEditingTool : public USingleSel
 
 public:
 	virtual void SetWorld(UWorld* World) { TargetWorld = World; }
-	virtual void SetAssetAPI(IAssetGenerationAPI* AssetAPIIn) { AssetAPI = AssetAPIIn; }
 
 protected:
 	UWorld* TargetWorld = nullptr;
-	IAssetGenerationAPI* AssetAPI = nullptr;
-
-
 
 	//
 	// Mesh Selection support
diff --git a/Engine/Plugins/Experimental/MeshModelingToolset/Source/ModelingComponents/Public/ModelingObjectsCreationAPI.h b/Engine/Plugins/Experimental/MeshModelingToolset/Source/ModelingComponents/Public/ModelingObjectsCreationAPI.h
new file mode 100644
--- /dev/null
+++ b/Engine/Plugins/Experimental/MeshModelingToolset/Source/ModelingComponents/Public/ModelingObjectsCreationAPI.h
@@ -0,0 +1,386 @@
+// Copyright Epic Games, Inc. All Rights Reserved.
+
+#pragma once
+
+
+#include "CoreMinimal.h"
+#include "BodySetupEnums.h"
+#include "MeshDescription.h"
+#include "DynamicMesh/DynamicMesh3.h"
+#include "Misc/Optional.h"
+
+#include "ModelingObjectsCreationAPI.generated.h"
+
+
+class UInteractiveToolManager;
+class AActor;
+class UPrimitiveComponent;
+class UMaterialInterface;
+class UTexture2D;
+
+
+/**
+ * Result code returned by UModelingObjectsCreationAPI functions
+ */
+UENUM(BlueprintType)
+enum class ECreateModelingObjectResult : uint8
+{
+	Ok,
+	Cancelled,
+	Failed_Unknown,
+	Failed_NoAPIFound,
+	Failed_InvalidWorld,
+	Failed_InvalidMesh,
+	Failed_InvalidTexture,
+	Failed_AssetCreationFailed
+};
+
+/**
+ * Types of possible source meshes stored in FCreateMeshObjectParams
+ */
+UENUM(BlueprintType)
+enum class ECreateMeshObjectSourceMeshType : uint8
+{
+	MeshDescription,
+	DynamicMesh
+};
+
+/**
+ * Hint for the type of mesh object a UModelingObjectsCreationAPI might create based
+ * on FCreateMeshObjectParams data. This can be used by clients to try to specify
+ * the type of object to emit, however there is no guarantee that an API implementation
+ * supports creating all types.
+ */
+UENUM(BlueprintType)
+enum class ECreateObjectTypeHint : uint8
+{
+	Undefined = 0,
+	Volume = 1,
+	StaticMesh = 2
+};
+
+
+/**
+ * FCreateMeshObjectParams is a collection of input data intended to be passed to
+ * UModelingObjectsCreationAPI::CreateMeshObject(). Not all data necessarily needs
+ * to be specified, this will depend on the particular implementation. The comments
+ * below are representative of how this data structure is used in the Tools and
+ * API implementation(s) provided with Unreal Engine, but end-user implementors
+ * could abuse these fields as necessary.
+ * 
+ * The definition of a "mesh object" is implementation-specific.
+ */
+USTRUCT(Blueprintable)
+struct MODELINGCOMPONENTS_API FCreateMeshObjectParams
+{
+	GENERATED_BODY()
+
+	//
+	// Base data
+	//
+
+	/** A Source Component the new mesh is based on, if such a Component exists */
+	UPROPERTY(Category = "CreateMeshObjectParams", EditAnywhere)
+	TObjectPtr<UPrimitiveComponent> SourceComponent = nullptr;
+
+	/** A suggested type for the newly-created Mesh (possibly ignored) */
+	UPROPERTY(Category = "CreateMeshObjectParams", EditAnywhere)
+	ECreateObjectTypeHint TypeHint = ECreateObjectTypeHint::Undefined;
+
+	/** An arbitrary integer that can be used to pass data to an API implementation */
+	UPROPERTY(Category = "CreateMeshObjectParams", EditAnywhere)
+	int32 TypeHintExtended = 0;
+
+	/** The World/Level the new mesh object should be created in (if known) */
+	UPROPERTY(Category = "CreateMeshObjectParams", EditAnywhere)
+	TObjectPtr<UWorld> TargetWorld = nullptr;
+
+	/** The 3D local-to-world transform for the new mesh object */
+	UPROPERTY(Category = "CreateMeshObjectParams", EditAnywhere)
+	FTransform Transform;
+
+	/** The base name of the new mesh object */
+	UPROPERTY(Category = "CreateMeshObjectParams", EditAnywhere)
+	FString BaseName;
+
+	//
+	// Materials settings
+	//
+
+	/** Materials for the new mesh object */
+	UPROPERTY(Category = "CreateMeshObjectParams", EditAnywhere)
+	TArray<TObjectPtr<UMaterialInterface>> Materials;
+
+	/** Optional Materials for a newly-created Mesh Asset, if this is applicable for the created mesh object */
+	UPROPERTY(Category = "CreateMeshObjectParams", EditAnywhere)
+	TArray<TObjectPtr<UMaterialInterface>> AssetMaterials;
+
+	//
+	// Collision settings, if applicable for the given mesh object
+	//
+
+	/** Specify whether the new mesh object should have collision support/data */
+	UPROPERTY(Category = "CreateMeshObjectParams", EditAnywhere)
+	bool bEnableCollision = true;
+
+	/** Which Collision mode to enable on the new mesh object, if supported */
+	UPROPERTY(Category = "CreateMeshObjectParams", EditAnywhere)
+	TEnumAsByte<enum ECollisionTraceFlag> CollisionMode = ECollisionTraceFlag::CTF_UseComplexAsSimple;
+
+
+	//
+	// Mesh Build Options, if this is applicable for the given mesh object
+	// (Currently somewhat specific to Assets in the Editor)
+	//
+
+	/** Specify whether normals should be automatically recomputed for this new mesh object */
+	UPROPERTY(Category = "CreateMeshObjectParams", EditAnywhere)
+	bool bEnableRecomputeNormals = false;
+
+	/** Specify whether tangents should be automatically recomputed for this new mesh object */
+	UPROPERTY(Category = "CreateMeshObjectParams", EditAnywhere)
+	bool bEnableRecomputeTangents = false;
+
+	/** Specify whether Nanite should be enabled on this new mesh object */
+	UPROPERTY(Category = "CreateMeshObjectParams", EditAnywhere)
+	bool bEnableNanite = false;
+
+	/** Specify the Nanite proxy triangle percentage for this new mesh object */
+	UPROPERTY(Category = "CreateMeshObjectParams", EditAnywhere)
+	float NaniteProxyTrianglePercent = 100.0f;
+
+
+	//
+	// The Mesh Object should be created based on the mesh data structures below.
+	// The assumption is that only one of these mesh data structures will be initialized.
+	// Not currently exposed to BP because the types are not BP types.
+	//
+
+	ECreateMeshObjectSourceMeshType MeshType = ECreateMeshObjectSourceMeshType::MeshDescription;
+	TOptional<FMeshDescription> MeshDescription;
+	TOptional<UE::Geometry::FDynamicMesh3> DynamicMesh;
+
+	void SetMesh(FMeshDescription&& MeshDescriptionIn);
+	void SetMesh(const UE::Geometry::FDynamicMesh3* DynamicMeshIn);
+	void SetMesh(UE::Geometry::FDynamicMesh3&& DynamicMeshIn);
+
+};
+
+
+/**
+ * FCreateMeshObjectResult is returned by UModelingObjectsCreationAPI::CreateMeshObject()
+ * to indicate success/failure and provide information about created mesh objects
+ */
+USTRUCT(BlueprintType)
+struct MODELINGCOMPONENTS_API FCreateMeshObjectResult
+{
+	GENERATED_BODY()
+
+	/** Success/Failure status for the requested operation */
+	UPROPERTY(Category = "CreateMeshObjectResult", VisibleAnywhere)
+	ECreateModelingObjectResult ResultCode;
+
+	/** A pointer to a newly-created Actor for the mesh object, if applicable (eg StaticMeshActor) */
+	UPROPERTY(Category="CreateMeshObjectResult", VisibleAnywhere)
+	TObjectPtr<AActor> NewActor = nullptr;
+
+	/** A pointer to a newly-created PrimitiveComponent for the mesh object, if applicable (eg StaticMeshComponent) */
+	UPROPERTY(Category = "CreateMeshObjectResult", VisibleAnywhere)
+	TObjectPtr<UPrimitiveComponent> NewComponent = nullptr;
+
+	/** A pointer to a newly-created Asset for the mesh object, if applicable (eg StaticMeshAsset) */
+	UPROPERTY(Category = "CreateMeshObjectResult", VisibleAnywhere)
+	TObjectPtr<UObject> NewAsset = nullptr;
+
+
+	bool IsOK() const { return ResultCode == ECreateModelingObjectResult::Ok; }
+};
+
+
+
+
+
+/**
+ * FCreateTextureObjectParams is a collection of input data intended to be passed to
+ * UModelingObjectsCreationAPI::CreateTextureObject(). Not all data necessarily needs
+ * to be specified, this will depend on the particular implementation. The comments
+ * below are representative of how this data structure is used in the Tools and
+ * API implementation(s) provided with Unreal Engine, but end-user implementors
+ * could abuse these fields as necessary.
+ *
+ * The definition of a "texture object" is implementation-specific. 
+ * In the UE Editor this is generally a UTexture2D
+ */
+USTRUCT(Blueprintable)
+struct MODELINGCOMPONENTS_API FCreateTextureObjectParams
+{
+	GENERATED_BODY()
+
+	//
+	// Base data
+	//
+
+	/** An arbitrary integer that can be used to pass data to an API implementation */
+	UPROPERTY(Category = "CreateTextureObjectParams", EditAnywhere)
+	int32 TypeHintExtended = 0;
+
+	/** 
+	 * The World/Level the new texture object should be created in (if known).
+	 * Note that Textures generally do not exist as objects in a Level. 
+	 * However, it may be necessary to store the texture in a path relative to the
+	 * level (for example if the level is in a Plugin, this would be necessary in-Editor)
+	 */
+	UPROPERTY(Category = "CreateTextureObjectParams", EditAnywhere)
+	TObjectPtr<UWorld> TargetWorld = nullptr;
+
+	/** An object to store the Texture relative to. For example the texture could be stored at the same path. */
+	UPROPERTY(Category = "CreateTextureObjectParams", EditAnywhere)
+	TObjectPtr<UObject> StoreRelativeToObject = nullptr;
+
+	/** The base name of the new mesh object */
+	UPROPERTY(Category = "CreateTextureObjectParams", EditAnywhere)
+	FString BaseName;
+
+	//
+	// input data
+	//
+
+	/** 
+	 * Texture source data. Generally assumed that this is a Texture created in the Transient package
+	 * that is intended to be saved in a permanent package.
+	 */
+	UPROPERTY(Category = "CreateTextureObjectParams", EditAnywhere)
+	TObjectPtr<UTexture2D> GeneratedTransientTexture = nullptr;
+};
+
+
+/**
+ * FCreateTextureObjectResult is returned by UModelingObjectsCreationAPI::CreateTextureObject()
+ * to indicate success/failure and provide information about created texture objects
+ */
+USTRUCT(BlueprintType)
+struct MODELINGCOMPONENTS_API FCreateTextureObjectResult
+{
+	GENERATED_BODY()
+
+	/** Success/Failure status for the requested operation */
+	UPROPERTY(Category = "CreateTextureObjectResult", VisibleAnywhere)
+	ECreateModelingObjectResult ResultCode;
+
+	/** A pointer to a newly-created Asset for the texture object, if applicable (eg StaticMeshAsset) */
+	UPROPERTY(Category = "CreateTextureObjectResult", VisibleAnywhere)
+	TObjectPtr<UObject> NewAsset = nullptr;
+
+
+	bool IsOK() const { return ResultCode == ECreateModelingObjectResult::Ok; }
+};
+
+
+
+
+
+
+/**
+ * UModelingObjectsCreationAPI is a base interface for functions that can be used to
+ * create various types of objects from Modeling Tools, or other sources. The "type" is
+ * very generic here - "Mesh", "Texture", etc - because this API is meant to provide
+ * an abstraction for Tools to emit different types of objects in different situations.
+ * For example an Tool might emit StaticMesh Asset/Actors in-Editor, but ProceduralMeshComponents at Runtime.
+ * 
+ * The creation inputs are specified via the structs above (eg FCreateMeshObjectParams, FCreateTextureObjectParams),
+ * which are very extensive, kitchen-sink sort of structs. Generally "New Mesh Object"
+ * creation behavior will be very complex and so this API is really just a way to route
+ * the data, and very few guarantees can be made about any specific implementation.
+ * 
+ * The assumed (but not really required) usage of instances of this type are that they
+ * will be registered with an InteractiveToolsContext's ContextObjectStore, and then 
+ * fetched from there by Tools/Algorithms/etc that need to use these capabilities can
+ * use the UE::Modeling::CreateXObject() helper functions below. However the interface
+ * does not have any dependencies on this usage model. 
+ * 
+ * See UEditorModelingObjectsCreationAPI for an example implementation suitable for in-Editor use.
+ */
+UCLASS(Abstract, BlueprintType)
+class MODELINGCOMPONENTS_API UModelingObjectsCreationAPI : public UObject
+{
+	GENERATED_BODY()
+public:
+	/**
+	 * Create a new mesh object based on the data in CreateMeshParams
+	 * @return a results data structure, containing a result code and information about any new objects created
+	 */
+	UFUNCTION(BlueprintCallable, Category = "Modeling Objects")
+	virtual FCreateMeshObjectResult CreateMeshObject(const FCreateMeshObjectParams& CreateMeshParams) { return FCreateMeshObjectResult{ ECreateModelingObjectResult::Failed_Unknown }; }
+
+	/**
+	 * Create a new texture object based on the data in CreateTexParams
+	 * @return a results data structure, containing a result code and information about any new objects created
+	 */
+	UFUNCTION(BlueprintCallable, Category = "Modeling Objects")
+	virtual FCreateTextureObjectResult CreateTextureObject(const FCreateTextureObjectParams& CreateTexParams) { return FCreateTextureObjectResult{ ECreateModelingObjectResult::Failed_Unknown }; }
+
+
+	//
+	// Non-UFunction variants that support std::move operators
+	//
+
+	/**
+	 * If this function returns true, then the CreateMeshObject() and CreateTextureObject() that take && types can be called.
+	 * This can reduce data copying, eg if the mesh data can be directly moved into the output mesh object.
+	 */
+	virtual bool HasMoveVariants() const { return false; }
+
+	virtual FCreateMeshObjectResult CreateMeshObject(FCreateMeshObjectParams&& CreateMeshParams) { return FCreateMeshObjectResult{ ECreateModelingObjectResult::Failed_Unknown }; }
+	virtual FCreateTextureObjectResult CreateTextureObject(FCreateTextureObjectParams&& CreateTexParams) { return FCreateTextureObjectResult{ ECreateModelingObjectResult::Failed_Unknown }; }
+
+};
+
+
+
+
+namespace UE
+{
+namespace Modeling
+{
+
+
+/**
+ * Create a new mesh object based on the data in CreateMeshParams.
+ * This is a convenience function that will try to locate a UModelingObjectsCreationAPI instance in the ToolManager's ContextObjectStore,
+ * and then call UModelingObjectsCreationAPI::CreateMeshObject()
+ * @return a results data structure, containing a result code and information about any new objects created
+ */
+MODELINGCOMPONENTS_API FCreateMeshObjectResult CreateMeshObject(UInteractiveToolManager* ToolManager, FCreateMeshObjectParams&& CreateMeshParams);
+
+
+/**
+ * Create a new texture object based on the data in CreateTexParams.
+ * This is a convenience function that will try to locate a UModelingObjectsCreationAPI instance in the ToolManager's ContextObjectStore,
+ * and then call UModelingObjectsCreationAPI::CreateTextureObject()
+ * @return a results data structure, containing a result code and information about any new objects created
+ */
+MODELINGCOMPONENTS_API FCreateTextureObjectResult CreateTextureObject(UInteractiveToolManager* ToolManager, FCreateTextureObjectParams&& CreateTexParams);
+
+
+
+/**
+ * Strip the appended auto-generated _UUID suffix on the given string, if we can detect one
+ * @param InputName input string that may have _UUID suffix
+ * @return InputName without _UUID suffix
+ */
+MODELINGCOMPONENTS_API FString StripGeneratedAssetSuffixFromName( FString InputName );
+
+/**
+ * Look up the name for the Asset underlying the given Component, if there is one.
+ * Optionally Strip off the appended auto-generated UUID string if we can detect one.
+ * If there is not a known underlying asset, returns the Component's name
+ * @param Component the Component we want the Asset Base Name for
+ * @param bRemoveAutoGeneratedSuffixes strip off generated UUID suffix if one is detected (default true)
+ * @return the Name for the Component
+ */
+MODELINGCOMPONENTS_API FString GetComponentAssetBaseName(
+	UPrimitiveComponent* Component,
+	bool bRemoveAutoGeneratedSuffixes = true);
+
+}
+}
diff --git a/Engine/Plugins/Experimental/MeshModelingToolset/Source/ModelingComponentsEditorOnly/ModelingComponentsEditorOnly.Build.cs b/Engine/Plugins/Experimental/MeshModelingToolset/Source/ModelingComponentsEditorOnly/ModelingComponentsEditorOnly.Build.cs
--- a/Engine/Plugins/Experimental/MeshModelingToolset/Source/ModelingComponentsEditorOnly/ModelingComponentsEditorOnly.Build.cs
+++ b/Engine/Plugins/Experimental/MeshModelingToolset/Source/ModelingComponentsEditorOnly/ModelingComponentsEditorOnly.Build.cs
@@ -47,6 +47,7 @@ public class ModelingComponentsEditorOnly : ModuleRules
 				"Engine",
 				"InputCore",
                 "RenderCore",
+				"PhysicsCore",
                 "RHI",
 
 				"AssetTools",
diff --git a/Engine/Plugins/Experimental/MeshModelingToolset/Source/ModelingComponentsEditorOnly/Private/EditorModelingObjectsCreationAPI.cpp b/Engine/Plugins/Experimental/MeshModelingToolset/Source/ModelingComponentsEditorOnly/Private/EditorModelingObjectsCreationAPI.cpp
new file mode 100644
--- /dev/null
+++ b/Engine/Plugins/Experimental/MeshModelingToolset/Source/ModelingComponentsEditorOnly/Private/EditorModelingObjectsCreationAPI.cpp
@@ -0,0 +1,245 @@
+// Copyright Epic Games, Inc. All Rights Reserved.
+
+#include "EditorModelingObjectsCreationAPI.h"
+#include "InteractiveToolsContext.h"
+#include "InteractiveToolManager.h"
+#include "ContextObjectStore.h"
+
+#include "AssetUtils/CreateStaticMeshUtil.h"
+#include "AssetUtils/CreateTexture2DUtil.h"
+
+#include "Engine/StaticMesh.h"
+#include "Components/StaticMeshComponent.h"
+#include "Engine/StaticMeshActor.h"
+
+using namespace UE::Geometry;
+
+
+UEditorModelingObjectsCreationAPI* UEditorModelingObjectsCreationAPI::Register(UInteractiveToolsContext* ToolsContext)
+{
+	if (ensure(ToolsContext))
+	{
+		UEditorModelingObjectsCreationAPI* CreationAPI = ToolsContext->ContextObjectStore->FindContext<UEditorModelingObjectsCreationAPI>();
+		if (CreationAPI)
+		{
+			return CreationAPI;
+		}
+		CreationAPI = NewObject<UEditorModelingObjectsCreationAPI>(ToolsContext);
+		ToolsContext->ContextObjectStore->AddContextObject(CreationAPI);
+		if (ensure(CreationAPI))
+		{
+			return CreationAPI;
+		}
+	}
+	return nullptr;
+}
+
+UEditorModelingObjectsCreationAPI* UEditorModelingObjectsCreationAPI::Find(UInteractiveToolsContext* ToolsContext)
+{
+	if (ensure(ToolsContext))
+	{
+		UEditorModelingObjectsCreationAPI* CreationAPI = ToolsContext->ContextObjectStore->FindContext<UEditorModelingObjectsCreationAPI>();
+		if (CreationAPI != nullptr)
+		{
+			return CreationAPI;
+		}
+	}
+	return nullptr;
+}
+
+bool UEditorModelingObjectsCreationAPI::Deregister(UInteractiveToolsContext* ToolsContext)
+{
+	if (ensure(ToolsContext))
+	{
+		UEditorModelingObjectsCreationAPI* CreationAPI = ToolsContext->ContextObjectStore->FindContext<UEditorModelingObjectsCreationAPI>();
+		if (CreationAPI != nullptr)
+		{
+			ToolsContext->ContextObjectStore->RemoveContextObject(CreationAPI);
+		}
+		return true;
+	}
+	return false;
+}
+
+
+FCreateMeshObjectResult UEditorModelingObjectsCreationAPI::CreateMeshObject(const FCreateMeshObjectParams& CreateMeshParams)
+{
+	// TODO: implement this path
+	check(false);
+	return FCreateMeshObjectResult{ ECreateModelingObjectResult::Failed_InvalidMesh };
+}
+
+
+FCreateTextureObjectResult UEditorModelingObjectsCreationAPI::CreateTextureObject(const FCreateTextureObjectParams& CreateTexParams)
+{
+	FCreateTextureObjectParams LocalParams = CreateTexParams;
+	return CreateTextureObject(MoveTemp(LocalParams));
+
+}
+
+
+FCreateMeshObjectResult UEditorModelingObjectsCreationAPI::CreateMeshObject(FCreateMeshObjectParams&& CreateMeshParams)
+{
+	if (!ensure(CreateMeshParams.TargetWorld)) { return FCreateMeshObjectResult{ ECreateModelingObjectResult::Failed_InvalidWorld }; }
+
+	UE::AssetUtils::FStaticMeshAssetOptions AssetOptions;
+
+	if (GetNewAssetPathNameCallback.IsBound())
+	{
+		AssetOptions.NewAssetPath = GetNewAssetPathNameCallback.Execute(CreateMeshParams.BaseName, CreateMeshParams.TargetWorld, FString());
+		if (AssetOptions.NewAssetPath.Len() == 0)
+		{
+			return FCreateMeshObjectResult{ ECreateModelingObjectResult::Cancelled };
+		}
+	}
+	else
+	{
+		AssetOptions.NewAssetPath = "/Game/" + CreateMeshParams.BaseName;
+	}
+
+	AssetOptions.NumSourceModels = 1;
+	AssetOptions.NumMaterialSlots = CreateMeshParams.Materials.Num();
+	AssetOptions.AssetMaterials = (CreateMeshParams.AssetMaterials.Num() == AssetOptions.NumMaterialSlots) ?
+		CreateMeshParams.AssetMaterials : CreateMeshParams.Materials;
+
+	AssetOptions.bEnableRecomputeNormals = CreateMeshParams.bEnableRecomputeNormals;
+	AssetOptions.bEnableRecomputeTangents = CreateMeshParams.bEnableRecomputeTangents;
+	AssetOptions.bGenerateNaniteEnabledMesh = CreateMeshParams.bEnableNanite;
+	AssetOptions.NaniteProxyTrianglePercent = CreateMeshParams.NaniteProxyTrianglePercent;
+
+	AssetOptions.bCreatePhysicsBody = CreateMeshParams.bEnableCollision;
+	AssetOptions.CollisionType = CreateMeshParams.CollisionMode;
+
+	if (CreateMeshParams.MeshType == ECreateMeshObjectSourceMeshType::DynamicMesh)
+	{
+		FDynamicMesh3* DynamicMesh = &CreateMeshParams.DynamicMesh.GetValue();
+		AssetOptions.SourceMeshes.DynamicMeshes.Add(DynamicMesh);
+	}
+	else if (CreateMeshParams.MeshType == ECreateMeshObjectSourceMeshType::MeshDescription)
+	{
+		FMeshDescription* MeshDescription = &CreateMeshParams.MeshDescription.GetValue();
+		AssetOptions.SourceMeshes.MoveMeshDescriptions.Add(MeshDescription);
+	}
+	else
+	{
+		return FCreateMeshObjectResult{ ECreateModelingObjectResult::Failed_InvalidMesh };
+	}
+
+	UE::AssetUtils::FStaticMeshResults ResultData;
+	UE::AssetUtils::ECreateStaticMeshResult AssetResult = UE::AssetUtils::CreateStaticMeshAsset(AssetOptions, ResultData);
+
+	if (AssetResult != UE::AssetUtils::ECreateStaticMeshResult::Ok)
+	{
+		return FCreateMeshObjectResult{ ECreateModelingObjectResult::Failed_AssetCreationFailed };
+	}
+
+	UStaticMesh* NewStaticMesh = ResultData.StaticMesh;
+
+	// create new StaticMeshActor
+	FRotator Rotation(0.0f, 0.0f, 0.0f);
+	FActorSpawnParameters SpawnInfo;
+	// @todo nothing here is specific to AStaticMeshActor...could we pass in a CDO and clone it instead of using SpawnActor?
+	AStaticMeshActor* StaticMeshActor = CreateMeshParams.TargetWorld->SpawnActor<AStaticMeshActor>(FVector::ZeroVector, Rotation, SpawnInfo);
+	StaticMeshActor->SetActorLabel(*CreateMeshParams.BaseName);
+	UStaticMeshComponent* StaticMeshComponent = StaticMeshActor->GetStaticMeshComponent();
+
+	// set the mesh
+
+	// this disconnects the component from various events
+	StaticMeshComponent->UnregisterComponent();
+	// Configure flags of the component. Is this necessary?
+	StaticMeshComponent->SetMobility(EComponentMobility::Movable);
+	StaticMeshComponent->bSelectable = true;
+	// replace the UStaticMesh in the component
+	StaticMeshComponent->SetStaticMesh(NewStaticMesh);
+
+	// set materials
+	for (int32 k = 0; k < CreateMeshParams.Materials.Num(); ++k)
+	{
+		StaticMeshComponent->SetMaterial(k, CreateMeshParams.Materials[k]);
+	}
+
+	// re-connect the component (?)
+	StaticMeshComponent->RegisterComponent();
+
+	NewStaticMesh->PostEditChange();
+
+	StaticMeshComponent->RecreatePhysicsState();
+
+	// update transform
+	StaticMeshActor->SetActorTransform(CreateMeshParams.Transform);
+
+	// emit result
+	FCreateMeshObjectResult ResultOut;
+	ResultOut.ResultCode = ECreateModelingObjectResult::Ok;
+	ResultOut.NewActor = StaticMeshActor;
+	ResultOut.NewComponent = StaticMeshComponent;
+	ResultOut.NewAsset = NewStaticMesh;
+
+	OnModelingMeshCreated.Broadcast(ResultOut);
+
+	return ResultOut;
+}
+
+
+
+
+FCreateTextureObjectResult UEditorModelingObjectsCreationAPI::CreateTextureObject(FCreateTextureObjectParams&& CreateTexParams)
+{
+	FString RelativeToObjectFolder;
+	if (CreateTexParams.StoreRelativeToObject != nullptr)
+	{
+		// find path to asset
+		UPackage* AssetOuterPackage = CastChecked<UPackage>(CreateTexParams.StoreRelativeToObject->GetOuter());
+		if (ensure(AssetOuterPackage))
+		{
+			FString AssetPackageName = AssetOuterPackage->GetName();
+			RelativeToObjectFolder = FPackageName::GetLongPackagePath(AssetPackageName);
+		}
+	}
+	else
+	{
+		if (!ensure(CreateTexParams.TargetWorld)) { return FCreateTextureObjectResult{ ECreateModelingObjectResult::Failed_InvalidWorld }; }
+	}
+
+	UE::AssetUtils::FTexture2DAssetOptions AssetOptions;
+
+	if (GetNewAssetPathNameCallback.IsBound())
+	{
+		AssetOptions.NewAssetPath = GetNewAssetPathNameCallback.Execute(CreateTexParams.BaseName, CreateTexParams.TargetWorld, RelativeToObjectFolder);
+		if (AssetOptions.NewAssetPath.Len() == 0)
+		{
+			return FCreateTextureObjectResult{ ECreateModelingObjectResult::Cancelled };
+		}
+	}
+	else
+	{
+		FString UseBaseFolder = (RelativeToObjectFolder.Len() > 0) ? RelativeToObjectFolder : TEXT("/Game");
+		AssetOptions.NewAssetPath = FPaths::Combine(UseBaseFolder, CreateTexParams.BaseName);
+	}
+
+	// currently we cannot create a new texture without an existing generated texture to store
+	if (!ensure(CreateTexParams.GeneratedTransientTexture))
+	{
+		return FCreateTextureObjectResult{ ECreateModelingObjectResult::Failed_InvalidTexture };
+	}
+
+	UE::AssetUtils::FTexture2DAssetResults ResultData;
+	UE::AssetUtils::ECreateTexture2DResult AssetResult = UE::AssetUtils::SaveGeneratedTexture2DAsset(
+		CreateTexParams.GeneratedTransientTexture, AssetOptions, ResultData);
+
+	if (AssetResult != UE::AssetUtils::ECreateTexture2DResult::Ok)
+	{
+		return FCreateTextureObjectResult{ ECreateModelingObjectResult::Failed_AssetCreationFailed };
+	}
+
+	// emit result
+	FCreateTextureObjectResult ResultOut;
+	ResultOut.ResultCode = ECreateModelingObjectResult::Ok;
+	ResultOut.NewAsset = ResultData.Texture;
+
+	OnModelingTextureCreated.Broadcast(ResultOut);
+
+	return ResultOut;
+
+}
\ No newline at end of file
diff --git a/Engine/Plugins/Experimental/MeshModelingToolset/Source/ModelingComponentsEditorOnly/Public/EditorModelingObjectsCreationAPI.h b/Engine/Plugins/Experimental/MeshModelingToolset/Source/ModelingComponentsEditorOnly/Public/EditorModelingObjectsCreationAPI.h
new file mode 100644
--- /dev/null
+++ b/Engine/Plugins/Experimental/MeshModelingToolset/Source/ModelingComponentsEditorOnly/Public/EditorModelingObjectsCreationAPI.h
@@ -0,0 +1,72 @@
+// Copyright Epic Games, Inc. All Rights Reserved.
+
+#pragma once
+
+#include "CoreMinimal.h"
+#include "ModelingObjectsCreationAPI.h"
+
+#include "EditorModelingObjectsCreationAPI.generated.h"
+
+class UInteractiveToolsContext;
+
+/**
+ * Implementation of UModelingObjectsCreationAPI suitable for use in UE Editor.
+ * - CreateMeshObject() currently creates a StaticMesh Asset/Actor
+ * - CreateTextureObject() currently creates a UTexture2D Asset
+ * 
+ * This is intended to be registered in the ToolsContext ContextObjectStore.
+ * Static utility functions ::Register() / ::Find() / ::Deregister() can be used to do this in a consistent way.
+ * 
+ * Several client-provided callbacks can be used to customize functionality (eg in Modeling Mode) 
+ *  - GetNewAssetPathNameCallback is called to determine an asset path. This can be used to do
+ *    things like pop up an interactive path-selection dialog, use project-defined paths, etc
+ *  - OnModelingMeshCreated is broadcast for each new created mesh object
+ *  - OnModelingTextureCreated is broadcast for each new created texture object
+ */
+UCLASS()
+class MODELINGCOMPONENTSEDITORONLY_API UEditorModelingObjectsCreationAPI : public UModelingObjectsCreationAPI
+{
+	GENERATED_BODY()
+public:
+
+	// UFUNCTION(BlueprintCallable, Category = "Modeling Objects")
+	virtual FCreateMeshObjectResult CreateMeshObject(const FCreateMeshObjectParams& CreateMeshParams) override;
+
+	// UFUNCTION(BlueprintCallable, Category = "Modeling Objects")
+	virtual FCreateTextureObjectResult CreateTextureObject(const FCreateTextureObjectParams& CreateTexParams) override;
+
+
+	//
+	// Non-UFunction variants that support std::move operators
+	//
+
+	virtual bool HasMoveVariants() const { return true; }
+
+	virtual FCreateMeshObjectResult CreateMeshObject(FCreateMeshObjectParams&& CreateMeshParams) override;
+
+	virtual FCreateTextureObjectResult CreateTextureObject(FCreateTextureObjectParams&& CreateTexParams) override;
+
+
+	//
+	// Callbacks that editor can hook into to handle asset creation
+	//
+
+	DECLARE_DELEGATE_RetVal_ThreeParams(FString, FGetAssetPathNameCallbackSignature, const FString& BaseName, const UWorld* TargetWorld, FString SuggestedFolder);
+	FGetAssetPathNameCallbackSignature GetNewAssetPathNameCallback;
+
+	DECLARE_MULTICAST_DELEGATE_OneParam(FModelingMeshCreatedSignature, const FCreateMeshObjectResult& CreatedInfo);
+	FModelingMeshCreatedSignature OnModelingMeshCreated;
+
+	DECLARE_MULTICAST_DELEGATE_OneParam(FModelingTextureCreatedSignature, const FCreateTextureObjectResult& CreatedInfo);
+	FModelingTextureCreatedSignature OnModelingTextureCreated;
+
+
+
+	//
+	// Utility functions to handle registration/unregistration
+	//
+
+	static UEditorModelingObjectsCreationAPI* Register(UInteractiveToolsContext* ToolsContext);
+	static UEditorModelingObjectsCreationAPI* Find(UInteractiveToolsContext* ToolsContext);
+	static bool Deregister(UInteractiveToolsContext* ToolsContext);
+};
diff --git a/Engine/Plugins/Experimental/MeshModelingToolset/Source/ModelingOperators/Private/ParameterizationOps/UVProjectionOp.cpp b/Engine/Plugins/Experimental/MeshModelingToolset/Source/ModelingOperators/Private/ParameterizationOps/UVProjectionOp.cpp
--- a/Engine/Plugins/Experimental/MeshModelingToolset/Source/ModelingOperators/Private/ParameterizationOps/UVProjectionOp.cpp
+++ b/Engine/Plugins/Experimental/MeshModelingToolset/Source/ModelingOperators/Private/ParameterizationOps/UVProjectionOp.cpp
@@ -2,13 +2,6 @@
 
 #include "ParameterizationOps/UVProjectionOp.h"
 
-#include "Engine/StaticMesh.h"
-
-#include "MeshDescriptionToDynamicMesh.h"
-#include "MeshSimplification.h"
-#include "MeshConstraintsUtil.h"
-#include "MatrixTypes.h"
-#include "Async/ParallelFor.h"
 #include "MeshNormals.h"
 #include "DynamicMeshEditor.h"
 #include "Parameterization/DynamicMeshUVEditor.h"
diff --git a/Engine/Plugins/Experimental/MeshModelingToolset/Source/ModelingOperators/Public/ParameterizationOps/UVProjectionOp.h b/Engine/Plugins/Experimental/MeshModelingToolset/Source/ModelingOperators/Public/ParameterizationOps/UVProjectionOp.h
--- a/Engine/Plugins/Experimental/MeshModelingToolset/Source/ModelingOperators/Public/ParameterizationOps/UVProjectionOp.h
+++ b/Engine/Plugins/Experimental/MeshModelingToolset/Source/ModelingOperators/Public/ParameterizationOps/UVProjectionOp.h
@@ -40,7 +40,7 @@ public:
 	// transform from Mesh into space of Projection Geometry
 	FTransform3d MeshToProjectionSpace;
 
-	EUVProjectionMethod ProjectionMethod;
+	EUVProjectionMethod ProjectionMethod = EUVProjectionMethod::Box;
 	FOrientedBox3d ProjectionBox;
 	float CylinderSplitAngle = 45.0f;
 
@@ -49,7 +49,7 @@ public:
 	int32 SmoothingRounds = 0;
 	double SmoothingAlpha = 0.5;
 
-	float CylinderProjectToTopOrBottomAngleThreshold;
+	float CylinderProjectToTopOrBottomAngleThreshold = 45.0f;
 
 	// position to use as UV origin
 	FVector2f UVOrigin = FVector2f(0.5f,0.5f);
diff --git a/Engine/Plugins/Experimental/ModelingToolsEditorMode/Source/ModelingToolsEditorMode/ModelingToolsEditorMode.Build.cs b/Engine/Plugins/Experimental/ModelingToolsEditorMode/Source/ModelingToolsEditorMode/ModelingToolsEditorMode.Build.cs
--- a/Engine/Plugins/Experimental/ModelingToolsEditorMode/Source/ModelingToolsEditorMode/ModelingToolsEditorMode.Build.cs
+++ b/Engine/Plugins/Experimental/ModelingToolsEditorMode/Source/ModelingToolsEditorMode/ModelingToolsEditorMode.Build.cs
@@ -54,6 +54,7 @@ public class ModelingToolsEditorMode : ModuleRules
                 "InteractiveToolsFramework",
 				"EditorInteractiveToolsFramework",
 				"ModelingComponents",
+				"ModelingComponentsEditorOnly",
 				"MeshModelingTools",
 				"MeshModelingToolsEditorOnly",
 				"MeshLODToolset",
diff --git a/Engine/Plugins/Experimental/ModelingToolsEditorMode/Source/ModelingToolsEditorMode/Private/ModelingModeAssetAPI.cpp b/Engine/Plugins/Experimental/ModelingToolsEditorMode/Source/ModelingToolsEditorMode/Private/ModelingModeAssetAPI.cpp
deleted file mode 100644
diff --git a/Engine/Plugins/Experimental/ModelingToolsEditorMode/Source/ModelingToolsEditorMode/Private/ModelingModeAssetUtils.cpp b/Engine/Plugins/Experimental/ModelingToolsEditorMode/Source/ModelingToolsEditorMode/Private/ModelingModeAssetUtils.cpp
new file mode 100644
--- /dev/null
+++ b/Engine/Plugins/Experimental/ModelingToolsEditorMode/Source/ModelingToolsEditorMode/Private/ModelingModeAssetUtils.cpp
@@ -0,0 +1,263 @@
+// Copyright Epic Games, Inc. All Rights Reserved.
+
+#include "ModelingModeAssetUtils.h"
+#include "ModelingToolsEditorModeSettings.h"
+
+#include "Engine/World.h"
+#include "AssetRegistryModule.h"
+#include "AssetToolsModule.h"
+#include "FileHelpers.h"
+#include "Misc/Paths.h"
+
+// for content-browser things
+#include "ContentBrowserModule.h"
+#include "IContentBrowserSingleton.h"
+
+
+#define LOCTEXT_NAMESPACE "ModelingModeAssetUtils"
+
+namespace UE
+{
+namespace Local
+{
+
+/**
+ * Generate a N-letter GUID string that contains only hex digits,
+ * and contains at least one letter and one number
+ */
+static FString GenerateRandomShortHexString(int32 NumChars = 8)
+{
+	int32 FailCount = 0;
+	while (FailCount++ < 10)
+	{
+		FGuid Guid = FGuid::NewGuid();
+		FString GuidString = Guid.ToString(EGuidFormats::UniqueObjectGuid).ToUpper();
+		FString Result;
+		int32 Digits = 0, Letters = 0;
+		for (int32 k = 0; k < GuidString.Len(); ++k)
+		{
+			TCHAR Character = GuidString[k];
+			if (FChar::IsHexDigit(Character))
+			{
+				Result.AppendChar(Character);
+				if (FChar::IsDigit(Character))
+				{
+					Digits++;
+				}
+				else
+				{
+					Letters++;
+				}
+			}
+			if (Result.Len() == NumChars)
+			{
+				if (Digits > 0 && Letters > 0)
+				{
+					return Result;
+				}
+				break;		// exit loop
+			}
+		}
+	}
+	return TEXT("BADGUID1");
+}
+
+
+
+FString GetWorldRelativeAssetRootPath(const UWorld* World)
+{
+	if (ensure(World->GetOutermost() != nullptr) == false)
+	{
+		return TEXT("/Game/");
+	}
+	FString WorldPackageName = World->GetOutermost()->GetName();
+	FString WorldPackageFolder = FPackageName::GetLongPackagePath(WorldPackageName);
+	return WorldPackageFolder;
+}
+
+
+FString GetActiveAssetFolderPath()
+{
+	IContentBrowserSingleton& ContentBrowser = FModuleManager::LoadModuleChecked<FContentBrowserModule>("ContentBrowser").Get();
+	return ContentBrowser.GetCurrentPath(EContentBrowserPathType::Internal);
+}
+
+
+FString MakeUniqueAssetName(const FString& FolderPath, const FString& AssetBaseName)
+{
+	FString UniquePackageName, UniqueAssetName;
+	FAssetToolsModule& AssetToolsModule = FModuleManager::LoadModuleChecked<FAssetToolsModule>("AssetTools");
+	AssetToolsModule.Get().CreateUniqueAssetName(
+		FolderPath + TEXT("/") + AssetBaseName, TEXT(""), UniquePackageName, UniqueAssetName);
+	return UniqueAssetName;
+}
+
+
+FString InteractiveSelectAssetPath(const FString& DefaultAssetName, const FText& DialogTitleMessage)
+{
+	IContentBrowserSingleton& ContentBrowser = FModuleManager::LoadModuleChecked<FContentBrowserModule>("ContentBrowser").Get();
+
+	FString UseDefaultAssetName = DefaultAssetName;
+	FString CurrentPath = GetActiveAssetFolderPath();
+	if (CurrentPath.IsEmpty() == false)
+	{
+		UseDefaultAssetName = MakeUniqueAssetName(CurrentPath, DefaultAssetName);
+	}
+
+	FSaveAssetDialogConfig Config;
+	Config.DefaultAssetName = UseDefaultAssetName;
+	Config.DialogTitleOverride = DialogTitleMessage;
+	Config.DefaultPath = CurrentPath;
+	return ContentBrowser.CreateModalSaveAssetDialog(Config);
+}
+
+
+
+}
+}
+
+
+
+
+
+
+FString UE::Modeling::GetNewAssetPathName(const FString& BaseNameIn, const UWorld* TargetWorld, FString SuggestedFolder)
+{
+	const UModelingToolsEditorModeSettings* Settings = GetDefault<UModelingToolsEditorModeSettings>();
+	EModelingModeAssetGenerationBehavior AutoGenMode = Settings->AssetGenerationMode;
+
+	FString PackageFolderPath = TEXT("/Game/");
+
+	if (SuggestedFolder.Len() > 0 && SuggestedFolder.StartsWith("/") )
+	{
+		// if we had a suggested folder, use it (todo: check that this folder exists? or can be created?)
+		PackageFolderPath = SuggestedFolder;
+	}
+	else
+	{
+		// Figure out the appropriate root path to use. If the level is in a plugin, this should
+		// be the plugin top-level path and not /Game.
+		if (Settings->AssetGenerationLocation == EModelingModeAssetGenerationLocation::AutoGeneratedWorldRelativeAssetPath)
+		{
+			PackageFolderPath = UE::Local::GetWorldRelativeAssetRootPath(TargetWorld);
+		}
+
+		// If the level is unsaved, then GetWorldRelativeAssetRootPath will return "/Temp". This is an undesirable
+		// place to even temporarily save assets because they must be manually moved elsewhere to save, and the /Temp folder 
+		// is somewhat troublesome to get to in the Editor. So if the flag below is set, then we will fall back to /Game
+		// instead. Projects with strict asset polices (due to P4/etc) may not want to allow this by default, so it is configurable.
+		if (PackageFolderPath.StartsWith("/Temp") && Settings->bStoreUnsavedLevelAssetsInTopLevelGameFolder)
+		{
+			PackageFolderPath = TEXT("/Game/");
+		}
+
+		// combine with fixed AutoGen path name if it is not empty
+		if (Settings->AutoGeneratedAssetPath.Len() > 0)
+		{
+			PackageFolderPath = FPaths::Combine(PackageFolderPath, Settings->AutoGeneratedAssetPath);
+		}
+
+		// append username-specific subfolder
+		if (Settings->bUsePerUserAutogenSubfolder)
+		{
+			FString UsernameString(FPlatformProcess::UserName());
+			if (UsernameString.Len() > 1)
+			{
+				PackageFolderPath = FPaths::Combine(PackageFolderPath, UsernameString);
+			}
+		}
+
+		// if we want to use the currently-visible asset browser path, try to find one (this can fail if no asset browser is visible/etc)
+		if (Settings->AssetGenerationLocation == EModelingModeAssetGenerationLocation::CurrentAssetBrowserPathIfAvailable)
+		{
+			FString CurrentAssetPath = UE::Local::GetActiveAssetFolderPath();
+			if (CurrentAssetPath.IsEmpty() == false)
+			{
+				PackageFolderPath = CurrentAssetPath;
+			}
+		}
+	}
+
+	FString ObjectBaseName = BaseNameIn;
+
+	// If we are in interactive mode, show the modal dialog and then get the path/name.
+	// If the user cancels, we are going to discard the asset
+	if (AutoGenMode == EModelingModeAssetGenerationBehavior::InteractivePromptToSave)
+	{
+		FString SelectedPath = UE::Local::InteractiveSelectAssetPath(ObjectBaseName, LOCTEXT("GenerateStaticMeshActorPathDialogWarning", "Choose Folder Path and Name for New Asset. Cancel to Discard New Asset."));
+		if (SelectedPath.IsEmpty() == false)
+		{
+			PackageFolderPath = FPaths::GetPath(SelectedPath);
+			ObjectBaseName = FPaths::GetBaseFilename(SelectedPath, true);
+		}
+		else
+		{
+			return FString();
+		}
+	}
+
+	FString UseBaseName = ObjectBaseName;
+	if (Settings->bAppendRandomStringToName)
+	{
+		FString GuidString = UE::Local::GenerateRandomShortHexString();
+		UseBaseName = FString::Printf(TEXT("%s_%s"), *UseBaseName, *GuidString);
+	}
+
+
+	FString PackageNameOut, AssetNameOut;
+	FAssetToolsModule& AssetToolsModule = FModuleManager::LoadModuleChecked<FAssetToolsModule>("AssetTools");
+	AssetToolsModule.Get().CreateUniqueAssetName(
+		FPaths::Combine(PackageFolderPath, UseBaseName), TEXT(""),
+		PackageNameOut, AssetNameOut);
+
+	return PackageNameOut;
+}
+
+
+
+bool UE::Modeling::AutoSaveAsset(UObject* Asset)
+{
+	UPackage* AssetPackage = Asset->GetPackage();
+	if (ensure(AssetPackage) == false)
+	{
+		return false;
+	}
+
+	Asset->MarkPackageDirty();
+	FAssetRegistryModule::AssetCreated(Asset);
+
+	TArray<UPackage*> PackagesToSave;
+	PackagesToSave.Add(AssetPackage);
+
+	bool bCheckDirty = true;
+	bool bPromptToSave = false;		// because we are autosaving
+	FEditorFileUtils::EPromptReturnCode Result = FEditorFileUtils::PromptForCheckoutAndSave(PackagesToSave, bCheckDirty, bPromptToSave);
+	ensure(Result == FEditorFileUtils::PR_Success);
+	return (Result == FEditorFileUtils::PR_Success);
+}
+
+
+
+
+
+void UE::Modeling::OnNewAssetCreated(UObject* Asset)
+{
+	const UModelingToolsEditorModeSettings* Settings = GetDefault<UModelingToolsEditorModeSettings>();
+	EModelingModeAssetGenerationBehavior AutoGenMode = Settings->AssetGenerationMode;
+
+	// Auto-save the new asset if the user has requested this behavior, or if an interactive
+	// save/path dialog was shown earlier in the process (InteractivePromptToSave mode)
+	if (AutoGenMode == EModelingModeAssetGenerationBehavior::AutoGenerateAndAutosave
+		 || AutoGenMode == EModelingModeAssetGenerationBehavior::InteractivePromptToSave)
+	{
+		AutoSaveAsset(Asset);
+	}
+	else if (AutoGenMode == EModelingModeAssetGenerationBehavior::AutoGenerateButDoNotAutosave)
+	{
+		Asset->MarkPackageDirty();
+		FAssetRegistryModule::AssetCreated(Asset);
+	}
+}
+
+
+#undef LOCTEXT_NAMESPACE
\ No newline at end of file
diff --git a/Engine/Plugins/Experimental/ModelingToolsEditorMode/Source/ModelingToolsEditorMode/Private/ModelingToolsEditorMode.cpp b/Engine/Plugins/Experimental/ModelingToolsEditorMode/Source/ModelingToolsEditorMode/Private/ModelingToolsEditorMode.cpp
--- a/Engine/Plugins/Experimental/ModelingToolsEditorMode/Source/ModelingToolsEditorMode/Private/ModelingToolsEditorMode.cpp
+++ b/Engine/Plugins/Experimental/ModelingToolsEditorMode/Source/ModelingToolsEditorMode/Private/ModelingToolsEditorMode.cpp
@@ -6,6 +6,7 @@
 #include "ModelingToolsEditorModeToolkit.h"
 #include "Toolkits/ToolkitManager.h"
 #include "ToolTargetManager.h"
+#include "ContextObjectStore.h"
 #include "ToolTargets/StaticMeshComponentToolTarget.h"
 #include "ToolTargets/VolumeDynamicMeshToolTarget.h"
 #include "ToolTargets/VolumeMeshDescriptionToolTarget.h"
@@ -107,7 +108,8 @@
 
 #include "ModelingToolsActions.h"
 #include "ModelingToolsManagerActions.h"
-#include "ModelingModeAssetAPI.h"
+#include "ModelingModeAssetUtils.h"
+#include "EditorModelingObjectsCreationAPI.h"
 
 #define LOCTEXT_NAMESPACE "UModelingToolsEditorMode"
 
@@ -294,8 +296,6 @@ void UModelingToolsEditorMode::Enter()
 {
 	UEdMode::Enter();
 
-	ModelingModeAssetGenerationAPI = MakeShareable(new FModelingModeAssetAPI(ToolsContext->GetAssetAPI()));
-
 	// Register builders for tool targets that the mode uses.
 	ToolsContext->TargetManager->AddTargetFactory(NewObject<UStaticMeshComponentToolTargetFactory>(ToolsContext->TargetManager));
 	ToolsContext->TargetManager->AddTargetFactory(NewObject<UVolumeDynamicMeshToolTargetFactory>(ToolsContext->TargetManager));
@@ -307,6 +307,24 @@ void UModelingToolsEditorMode::Enter()
 	// register gizmo helper
 	UE::TransformGizmoUtil::RegisterTransformGizmoContextObject(ToolsContext.Get());
 
+	// register object creation api
+	UEditorModelingObjectsCreationAPI* ModelCreationAPI = UEditorModelingObjectsCreationAPI::Register(ToolsContext.Get());
+	if (ModelCreationAPI)
+	{
+		ModelCreationAPI->GetNewAssetPathNameCallback.BindLambda([](const FString& BaseName, const UWorld* TargetWorld, FString SuggestedFolder)
+		{
+			return UE::Modeling::GetNewAssetPathName(BaseName, TargetWorld, SuggestedFolder);
+		});
+		MeshCreatedEventHandle = ModelCreationAPI->OnModelingMeshCreated.AddLambda([](const FCreateMeshObjectResult& CreatedInfo) 
+		{ 
+			UE::Modeling::OnNewAssetCreated(CreatedInfo.NewAsset);
+		});
+		TextureCreatedEventHandle = ModelCreationAPI->OnModelingTextureCreated.AddLambda([](const FCreateTextureObjectResult& CreatedInfo)
+		{
+			UE::Modeling::OnNewAssetCreated(CreatedInfo.NewAsset);
+		});
+	}
+
 	const FModelingToolsManagerCommands& ToolManagerCommands = FModelingToolsManagerCommands::Get();
 
 	// register tool set
@@ -320,7 +338,6 @@ void UModelingToolsEditorMode::Enter()
 								  UAddPrimitiveToolBuilder::EMakeMeshShapeType ShapeTypeIn)
 	{
 		auto AddPrimitiveToolBuilder = NewObject<UAddPrimitiveToolBuilder>();
-		AddPrimitiveToolBuilder->AssetAPI = ModelingModeAssetGenerationAPI.Get();
 		AddPrimitiveToolBuilder->ShapeType = ShapeTypeIn;
 		RegisterTool(UICommand, ToolIdentifier, AddPrimitiveToolBuilder);
 	};
@@ -356,27 +373,21 @@ void UModelingToolsEditorMode::Enter()
 	// make shape tools
 	//
 	auto AddPatchToolBuilder = NewObject<UAddPatchToolBuilder>();
-	AddPatchToolBuilder->AssetAPI = ModelingModeAssetGenerationAPI.Get();
 	RegisterTool(ToolManagerCommands.BeginAddPatchTool, TEXT("BeginAddPatchTool"), AddPatchToolBuilder);
 
 	auto RevolveBoundaryToolBuilder = NewObject<URevolveBoundaryToolBuilder>();
-	RevolveBoundaryToolBuilder->AssetAPI = ModelingModeAssetGenerationAPI.Get();
 	RegisterTool(ToolManagerCommands.BeginRevolveBoundaryTool, TEXT("BeginRevolveBoundaryTool"), RevolveBoundaryToolBuilder);
 
 	auto DrawPolygonToolBuilder = NewObject<UDrawPolygonToolBuilder>();
-	DrawPolygonToolBuilder->AssetAPI = ModelingModeAssetGenerationAPI.Get();
 	RegisterTool(ToolManagerCommands.BeginDrawPolygonTool, TEXT("BeginDrawPolygonTool"), DrawPolygonToolBuilder);
 
 	auto DrawPolyPathToolBuilder = NewObject<UDrawPolyPathToolBuilder>();
-	DrawPolyPathToolBuilder->AssetAPI = ModelingModeAssetGenerationAPI.Get();
 	RegisterTool(ToolManagerCommands.BeginDrawPolyPathTool, TEXT("BeginDrawPolyPathTool"), DrawPolyPathToolBuilder);
 
 	auto DrawAndRevolveToolBuilder = NewObject<UDrawAndRevolveToolBuilder>();
-	DrawAndRevolveToolBuilder->AssetAPI = ModelingModeAssetGenerationAPI.Get();
 	RegisterTool(ToolManagerCommands.BeginDrawAndRevolveTool, TEXT("BeginDrawAndRevolveTool"), DrawAndRevolveToolBuilder);
 
 	auto ShapeSprayToolBuilder = NewObject<UShapeSprayToolBuilder>();
-	ShapeSprayToolBuilder->AssetAPI = ModelingModeAssetGenerationAPI.Get();
 	RegisterTool(ToolManagerCommands.BeginShapeSprayTool, TEXT("BeginShapeSprayTool"), ShapeSprayToolBuilder);
 
 
@@ -410,21 +421,17 @@ void UModelingToolsEditorMode::Enter()
 	RegisterTool(ToolManagerCommands.BeginSubdividePolyTool, TEXT("BeginSubdividePolyTool"), NewObject<USubdividePolyToolBuilder>());
 
 	UCombineMeshesToolBuilder* CombineMeshesToolBuilder = NewObject<UCombineMeshesToolBuilder>();
-	CombineMeshesToolBuilder->AssetAPI = ModelingModeAssetGenerationAPI.Get();
 	RegisterTool(ToolManagerCommands.BeginCombineMeshesTool, TEXT("BeginCombineMeshesTool"), CombineMeshesToolBuilder);
 
 	UCombineMeshesToolBuilder* DuplicateMeshesToolBuilder = NewObject<UCombineMeshesToolBuilder>();
-	DuplicateMeshesToolBuilder->AssetAPI = ModelingModeAssetGenerationAPI.Get();
 	DuplicateMeshesToolBuilder->bIsDuplicateTool = true;
 	RegisterTool(ToolManagerCommands.BeginDuplicateMeshesTool, TEXT("BeginDuplicateMeshesTool"), DuplicateMeshesToolBuilder);
 
 
 	ULODManagerToolBuilder* LODManagerToolBuilder = NewObject<ULODManagerToolBuilder>();
-	LODManagerToolBuilder->AssetAPI = ModelingModeAssetGenerationAPI.Get();
 	RegisterTool(ToolManagerCommands.BeginLODManagerTool, TEXT("BeginLODManagerTool"), LODManagerToolBuilder);
 
 	UGenerateStaticMeshLODAssetToolBuilder* GenerateSMLODToolBuilder = NewObject<UGenerateStaticMeshLODAssetToolBuilder>();
-	GenerateSMLODToolBuilder->AssetAPI = ModelingModeAssetGenerationAPI.Get();
 	RegisterTool(ToolManagerCommands.BeginGenerateStaticMeshLODAssetTool, TEXT("BeginGenerateStaticMeshLODAssetTool"), GenerateSMLODToolBuilder);
 
 
@@ -440,118 +447,90 @@ void UModelingToolsEditorMode::Enter()
 	RegisterTool(ToolManagerCommands.BeginProjectToTargetTool, TEXT("BeginProjectToTargetTool"), NewObject<UProjectToTargetToolBuilder>());
 	RegisterTool(ToolManagerCommands.BeginSimplifyMeshTool, TEXT("BeginSimplifyMeshTool"), NewObject<USimplifyMeshToolBuilder>());
 
-	auto GroupEdgeInsertionToolBuilder = NewObject<UGroupEdgeInsertionToolBuilder>();
-	RegisterTool(ToolManagerCommands.BeginGroupEdgeInsertionTool, TEXT("BeginGroupEdgeInsertionTool"), GroupEdgeInsertionToolBuilder);
+	RegisterTool(ToolManagerCommands.BeginGroupEdgeInsertionTool, TEXT("BeginGroupEdgeInsertionTool"), NewObject<UGroupEdgeInsertionToolBuilder>());
 
-	auto EdgeLoopInsertionToolBuilder = NewObject<UEdgeLoopInsertionToolBuilder>();
-	RegisterTool(ToolManagerCommands.BeginEdgeLoopInsertionTool, TEXT("BeginEdgeLoopInsertionTool"), EdgeLoopInsertionToolBuilder);
+	RegisterTool(ToolManagerCommands.BeginEdgeLoopInsertionTool, TEXT("BeginEdgeLoopInsertionTool"), NewObject<UEdgeLoopInsertionToolBuilder>());
 
 	auto EditNormalsToolBuilder = NewObject<UEditNormalsToolBuilder>();
-	EditNormalsToolBuilder->AssetAPI = ModelingModeAssetGenerationAPI.Get();
 	RegisterTool(ToolManagerCommands.BeginEditNormalsTool, TEXT("BeginEditNormalsTool"), EditNormalsToolBuilder);
 
-	auto TangentsToolBuilder = NewObject<UMeshTangentsToolBuilder>();
-	RegisterTool(ToolManagerCommands.BeginEditTangentsTool, TEXT("BeginEditTangentsTool"), TangentsToolBuilder);
+	RegisterTool(ToolManagerCommands.BeginEditTangentsTool, TEXT("BeginEditTangentsTool"), NewObject<UMeshTangentsToolBuilder>());
 
 	auto RemoveOccludedTrianglesToolBuilder = NewObject<URemoveOccludedTrianglesToolBuilder>();
-	RemoveOccludedTrianglesToolBuilder->AssetAPI = ModelingModeAssetGenerationAPI.Get();
 	RegisterTool(ToolManagerCommands.BeginRemoveOccludedTrianglesTool, TEXT("BeginRemoveOccludedTrianglesTool"), RemoveOccludedTrianglesToolBuilder);
 
-	auto HoleFillToolBuilder = NewObject<UHoleFillToolBuilder>();
-	RegisterTool(ToolManagerCommands.BeginHoleFillTool, TEXT("BeginHoleFillTool"), HoleFillToolBuilder);
+	RegisterTool(ToolManagerCommands.BeginHoleFillTool, TEXT("BeginHoleFillTool"), NewObject<UHoleFillToolBuilder>());
 
 	auto UVProjectionToolBuilder = NewObject<UUVProjectionToolBuilder>();
-	UVProjectionToolBuilder->AssetAPI = ModelingModeAssetGenerationAPI.Get();
 	RegisterTool(ToolManagerCommands.BeginUVProjectionTool, TEXT("BeginUVProjectionTool"), UVProjectionToolBuilder);
 
 	auto UVLayoutToolBuilder = NewObject<UUVLayoutToolBuilder>();
-	UVLayoutToolBuilder->AssetAPI = ModelingModeAssetGenerationAPI.Get();
 	RegisterTool(ToolManagerCommands.BeginUVLayoutTool, TEXT("BeginUVLayoutTool"), UVLayoutToolBuilder);
 
 	auto MergeMeshesToolBuilder = NewObject<UMergeMeshesToolBuilder>();
-	MergeMeshesToolBuilder->AssetAPI = ModelingModeAssetGenerationAPI.Get();
 	RegisterTool(ToolManagerCommands.BeginVoxelMergeTool, TEXT("BeginVoxelMergeTool"), MergeMeshesToolBuilder);
 
 	auto VoxelCSGMeshesToolBuilder = NewObject<UVoxelCSGMeshesToolBuilder>();
-	VoxelCSGMeshesToolBuilder->AssetAPI = ModelingModeAssetGenerationAPI.Get();
 	RegisterTool(ToolManagerCommands.BeginVoxelBooleanTool, TEXT("BeginVoxelBooleanTool"), VoxelCSGMeshesToolBuilder);
 
 	auto VoxelSolidifyMeshesToolBuilder = NewObject<UVoxelSolidifyMeshesToolBuilder>();
-	VoxelSolidifyMeshesToolBuilder->AssetAPI = ModelingModeAssetGenerationAPI.Get();
 	RegisterTool(ToolManagerCommands.BeginVoxelSolidifyTool, TEXT("BeginVoxelSolidifyTool"), VoxelSolidifyMeshesToolBuilder);
 
 	auto VoxelBlendMeshesToolBuilder = NewObject<UVoxelBlendMeshesToolBuilder>();
-	VoxelBlendMeshesToolBuilder->AssetAPI = ModelingModeAssetGenerationAPI.Get();
 	RegisterTool(ToolManagerCommands.BeginVoxelBlendTool, TEXT("BeginVoxelBlendTool"), VoxelBlendMeshesToolBuilder);
 
 	auto VoxelMorphologyMeshesToolBuilder = NewObject<UVoxelMorphologyMeshesToolBuilder>();
-	VoxelMorphologyMeshesToolBuilder->AssetAPI = ModelingModeAssetGenerationAPI.Get();
 	RegisterTool(ToolManagerCommands.BeginVoxelMorphologyTool, TEXT("BeginVoxelMorphologyTool"), VoxelMorphologyMeshesToolBuilder);
 
 	auto SelfUnionMeshesToolBuilder = NewObject<USelfUnionMeshesToolBuilder>();
-	SelfUnionMeshesToolBuilder->AssetAPI = ModelingModeAssetGenerationAPI.Get();
 	RegisterTool(ToolManagerCommands.BeginSelfUnionTool, TEXT("BeginSelfUnionTool"), SelfUnionMeshesToolBuilder);
 
 	auto CSGMeshesToolBuilder = NewObject<UCSGMeshesToolBuilder>();
-	CSGMeshesToolBuilder->AssetAPI = ModelingModeAssetGenerationAPI.Get();
 	RegisterTool(ToolManagerCommands.BeginMeshBooleanTool, TEXT("BeginMeshBooleanTool"), CSGMeshesToolBuilder);
 
 	auto CutMeshWithMeshToolBuilder = NewObject<UCutMeshWithMeshToolBuilder>();
-	CutMeshWithMeshToolBuilder->AssetAPI = ModelingModeAssetGenerationAPI.Get();
 	RegisterTool(ToolManagerCommands.BeginCutMeshWithMeshTool, TEXT("BeginCutMeshWithMeshTool"), CutMeshWithMeshToolBuilder);
 
 	auto TrimMeshesToolBuilder = NewObject<UCSGMeshesToolBuilder>();
 	TrimMeshesToolBuilder->bTrimMode = true;
-	TrimMeshesToolBuilder->AssetAPI = ModelingModeAssetGenerationAPI.Get();
 	RegisterTool(ToolManagerCommands.BeginMeshTrimTool, TEXT("BeginMeshTrimTool"), TrimMeshesToolBuilder);
 
-	auto BspConversionToolBuilder = NewObject<UBspConversionToolBuilder>();
-	BspConversionToolBuilder->AssetAPI = ModelingModeAssetGenerationAPI.Get();
-	RegisterTool(ToolManagerCommands.BeginBspConversionTool, TEXT("BeginBspConversionTool"), BspConversionToolBuilder);
+	RegisterTool(ToolManagerCommands.BeginBspConversionTool, TEXT("BeginBspConversionTool"), NewObject<UBspConversionToolBuilder>());
 
 	auto MeshToVolumeToolBuilder = NewObject<UMeshToVolumeToolBuilder>();
-	MeshToVolumeToolBuilder->AssetAPI = ModelingModeAssetGenerationAPI.Get();
 	RegisterTool(ToolManagerCommands.BeginMeshToVolumeTool, TEXT("BeginMeshToVolumeTool"), MeshToVolumeToolBuilder);
 
 	auto VolumeToMeshToolBuilder = NewObject<UVolumeToMeshToolBuilder>();
-	VolumeToMeshToolBuilder->AssetAPI = ModelingModeAssetGenerationAPI.Get();
 	RegisterTool(ToolManagerCommands.BeginVolumeToMeshTool, TEXT("BeginVolumeToMeshTool"), VolumeToMeshToolBuilder);
 
 	auto PlaneCutToolBuilder = NewObject<UPlaneCutToolBuilder>();
-	PlaneCutToolBuilder->AssetAPI = ModelingModeAssetGenerationAPI.Get();
 	RegisterTool(ToolManagerCommands.BeginPlaneCutTool, TEXT("BeginPlaneCutTool"), PlaneCutToolBuilder);
 
 	auto MirrorToolBuilder = NewObject<UMirrorToolBuilder>();
-	MirrorToolBuilder->AssetAPI = ModelingModeAssetGenerationAPI.Get();
 	RegisterTool(ToolManagerCommands.BeginMirrorTool, TEXT("BeginMirrorTool"), MirrorToolBuilder);
 
 	auto PolygonCutToolBuilder = NewObject<UPolygonOnMeshToolBuilder>();
 	RegisterTool(ToolManagerCommands.BeginPolygonCutTool, TEXT("BeginPolygonCutTool"), PolygonCutToolBuilder);
 
 	auto GlobalUVGenerateToolBuilder = NewObject<UParameterizeMeshToolBuilder>();
-	GlobalUVGenerateToolBuilder->AssetAPI = ModelingModeAssetGenerationAPI.Get();
 	GlobalUVGenerateToolBuilder->bDoAutomaticGlobalUnwrap = true;
 	RegisterTool(ToolManagerCommands.BeginGlobalUVGenerateTool, TEXT("BeginGlobalUVGenerateTool"), GlobalUVGenerateToolBuilder);
 
 	auto GroupUVGenerateToolBuilder = NewObject<UParameterizeMeshToolBuilder>();
-	GroupUVGenerateToolBuilder->AssetAPI = ModelingModeAssetGenerationAPI.Get();
 	GroupUVGenerateToolBuilder->bDoAutomaticGlobalUnwrap = false;
 	RegisterTool(ToolManagerCommands.BeginGroupUVGenerateTool, TEXT("BeginGroupUVGenerateTool"), GroupUVGenerateToolBuilder);
 
 	RegisterTool(ToolManagerCommands.BeginUVSeamEditTool, TEXT("BeginUVSeamEditTool"), NewObject< USeamSculptToolBuilder>());
 
 	auto MeshSelectionToolBuilder = NewObject<UMeshSelectionToolBuilder>();
-	MeshSelectionToolBuilder->AssetAPI = ModelingModeAssetGenerationAPI.Get();
 	RegisterTool(ToolManagerCommands.BeginMeshSelectionTool, TEXT("BeginMeshSelectionTool"), MeshSelectionToolBuilder);
 
 	auto EditMeshMaterialsToolBuilder = NewObject<UEditMeshMaterialsToolBuilder>();
-	EditMeshMaterialsToolBuilder->AssetAPI = ModelingModeAssetGenerationAPI.Get();
 	RegisterTool(ToolManagerCommands.BeginEditMeshMaterialsTool, TEXT("BeginEditMeshMaterialsTool"), EditMeshMaterialsToolBuilder);
 	
 	RegisterTool(ToolManagerCommands.BeginMeshAttributePaintTool, TEXT("BeginMeshAttributePaintTool"), NewObject<UMeshAttributePaintToolBuilder>());
 
 	auto BakeMeshAttributeMapsToolBuilder = NewObject<UBakeMeshAttributeMapsToolBuilder>();
-	BakeMeshAttributeMapsToolBuilder->AssetAPI = ModelingModeAssetGenerationAPI.Get();
 	RegisterTool(ToolManagerCommands.BeginBakeMeshAttributeMapsTool, TEXT("BeginBakeMeshAttributeMapsTool"), BakeMeshAttributeMapsToolBuilder);
 
 	// analysis tools
@@ -569,7 +548,6 @@ void UModelingToolsEditorMode::Enter()
 	//RegisterTool(ToolManagerCommands.BeginEditCollisionGeometryTool, TEXT("EditCollisionGeoTool"), NewObject<UEditCollisionGeometryToolBuilder>());
 
 	auto ExtractCollisionGeoToolBuilder = NewObject<UExtractCollisionGeometryToolBuilder>();
-	ExtractCollisionGeoToolBuilder->AssetAPI = ModelingModeAssetGenerationAPI.Get();
 	RegisterTool(ToolManagerCommands.BeginExtractCollisionGeometryTool, TEXT("BeginExtractCollisionGeometryTool"), ExtractCollisionGeoToolBuilder);
 
 
@@ -577,13 +555,11 @@ void UModelingToolsEditorMode::Enter()
 	// (experimental) hair tools
 
 	UGroomToMeshToolBuilder* GroomToMeshToolBuilder = NewObject<UGroomToMeshToolBuilder>();
-	GroomToMeshToolBuilder->AssetAPI = ModelingModeAssetGenerationAPI.Get();
 	RegisterTool(ToolManagerCommands.BeginGroomToMeshTool, TEXT("BeginGroomToMeshTool"), GroomToMeshToolBuilder);
 
 	RegisterTool(ToolManagerCommands.BeginGroomCardsEditorTool, TEXT("BeginGroomCardsEditorTool"), NewObject<UGroomCardsEditorToolBuilder>());
 
 	UGenerateLODMeshesToolBuilder* GenerateLODMeshesToolBuilder = NewObject<UGenerateLODMeshesToolBuilder>();
-	GenerateLODMeshesToolBuilder->AssetAPI = ModelingModeAssetGenerationAPI.Get();
 	RegisterTool(ToolManagerCommands.BeginGenerateLODMeshesTool, TEXT("BeginGenerateLODMeshesTool"), GenerateLODMeshesToolBuilder);
 
 	// PolyModeling tools
@@ -620,7 +596,7 @@ void UModelingToolsEditorMode::Enter()
 	{
 		FExtensionToolQueryInfo ExtensionQueryInfo;
 		ExtensionQueryInfo.ToolsContext = ToolsContext.Get();
-		ExtensionQueryInfo.AssetAPI = ModelingModeAssetGenerationAPI.Get();
+		ExtensionQueryInfo.AssetAPI = nullptr;
 
 		UE_LOG(LogTemp, Log, TEXT("ModelingMode: Found %d Tool Extension Modules"), Extensions.Num());
 		for (int32 k = 0; k < Extensions.Num(); ++k)
@@ -675,10 +651,21 @@ void UModelingToolsEditorMode::Exit()
 
 	StylusStateTracker = nullptr;
 
-	// TODO: cannot deregister currently because if another mode is also registering, it's Enter()
+	// TODO: cannot deregister currently because if another mode is also registering, its Enter()
 	// will be called before our Exit()
 	//UE::TransformGizmoUtil::DeregisterTransformGizmoContextObject(ToolsContext.Get());
 	
+	// TODO: cannot deregister currently because if another mode is also registering, its Enter()
+	// will be called before our Exit()
+	UEditorModelingObjectsCreationAPI* ObjectCreationAPI = UEditorModelingObjectsCreationAPI::Find(ToolsContext.Get());
+	if (ObjectCreationAPI)
+	{
+		ObjectCreationAPI->GetNewAssetPathNameCallback.Unbind();
+		ObjectCreationAPI->OnModelingMeshCreated.Remove(MeshCreatedEventHandle);
+		ObjectCreationAPI->OnModelingTextureCreated.Remove(TextureCreatedEventHandle);
+		//UEditorModelingObjectsCreationAPI::Deregister(ToolsContext.Get());		// cannot do currently because of shared ToolsContext, revisit in future
+	}
+
 	FModelingModeActionCommands::UnRegisterCommandBindings(Toolkit->GetToolkitCommands());
 
 	// clear realtime viewport override
diff --git a/Engine/Plugins/Experimental/ModelingToolsEditorMode/Source/ModelingToolsEditorMode/Private/StaticMeshComponentBuilder.h b/Engine/Plugins/Experimental/ModelingToolsEditorMode/Source/ModelingToolsEditorMode/Private/StaticMeshComponentBuilder.h
deleted file mode 100644
diff --git a/Engine/Plugins/Experimental/ModelingToolsEditorMode/Source/ModelingToolsEditorMode/Public/ModelingModeAssetAPI.h b/Engine/Plugins/Experimental/ModelingToolsEditorMode/Source/ModelingToolsEditorMode/Public/ModelingModeAssetAPI.h
deleted file mode 100644
diff --git a/Engine/Plugins/Experimental/ModelingToolsEditorMode/Source/ModelingToolsEditorMode/Public/ModelingModeAssetUtils.h b/Engine/Plugins/Experimental/ModelingToolsEditorMode/Source/ModelingToolsEditorMode/Public/ModelingModeAssetUtils.h
new file mode 100644
--- /dev/null
+++ b/Engine/Plugins/Experimental/ModelingToolsEditorMode/Source/ModelingToolsEditorMode/Public/ModelingModeAssetUtils.h
@@ -0,0 +1,45 @@
+// Copyright Epic Games, Inc. All Rights Reserved.
+
+#pragma once
+
+#include "CoreMinimal.h"
+
+class UWorld;
+
+
+namespace UE
+{
+namespace Modeling
+{
+
+//
+// Utility functions that can be used to plug in the client-configurable parts of UEditorModelingObjectsCreationAPI.
+// See ModelingToolsEditorMode for usage.
+//
+
+/**
+ * Determines path and name for a new Asset based on current mode settings, etc
+ * @param BaseName desired initial string for the name
+ * @param TargetWorld World the new Asset will be used in, which can determine its generated path based on various settings
+ * @param SuggestedFolder a suggested path for the asset, ignored if empty
+ * @return Path+Name for the new Asset (package)
+ */
+MODELINGTOOLSEDITORMODE_API FString GetNewAssetPathName(const FString& BaseName, const UWorld* TargetWorld, FString SuggestedFolder);
+
+
+/**
+ * Utility function that may auto-save a newly created asset depending on current mode settings,
+ * and posts notifications for the Editor
+ * @param Asset the new asset object (eg UStaticMesh, UTexture2D, etc)
+ */
+MODELINGTOOLSEDITORMODE_API void OnNewAssetCreated(UObject* Asset);
+
+
+/**
+ * Utility function that will attempt to auto-save an asset, and post Editor change notifications
+ * @param Asset the  asset object (eg UStaticMesh, UTexture2D, etc)
+ */
+MODELINGTOOLSEDITORMODE_API bool AutoSaveAsset(UObject* Asset);
+
+}
+}
\ No newline at end of file
diff --git a/Engine/Plugins/Experimental/ModelingToolsEditorMode/Source/ModelingToolsEditorMode/Public/ModelingModeToolExtensions.h b/Engine/Plugins/Experimental/ModelingToolsEditorMode/Source/ModelingToolsEditorMode/Public/ModelingModeToolExtensions.h
--- a/Engine/Plugins/Experimental/ModelingToolsEditorMode/Source/ModelingToolsEditorMode/Public/ModelingModeToolExtensions.h
+++ b/Engine/Plugins/Experimental/ModelingToolsEditorMode/Source/ModelingToolsEditorMode/Public/ModelingModeToolExtensions.h
@@ -6,11 +6,17 @@
 #include "Features/IModularFeature.h"
 
 class UInteractiveToolsContext;
-class FModelingModeAssetAPI;
 class FUICommandInfo;
 class UInteractiveToolBuilder;
 
 
+// dummy class for extension API below
+class FModelingModeAssetAPI
+{
+public:
+};
+
+
 /**
  * This struct is passed to IModelingModeToolExtension implementations to allow 
  * them to forward various Tools Context information to their ToolBuilders/etc
diff --git a/Engine/Plugins/Experimental/ModelingToolsEditorMode/Source/ModelingToolsEditorMode/Public/ModelingToolsEditorMode.h b/Engine/Plugins/Experimental/ModelingToolsEditorMode/Source/ModelingToolsEditorMode/Public/ModelingToolsEditorMode.h
--- a/Engine/Plugins/Experimental/ModelingToolsEditorMode/Source/ModelingToolsEditorMode/Public/ModelingToolsEditorMode.h
+++ b/Engine/Plugins/Experimental/ModelingToolsEditorMode/Source/ModelingToolsEditorMode/Public/ModelingToolsEditorMode.h
@@ -16,7 +16,6 @@ class FEditorComponentSourceFactory;
 class FEditorToolAssetAPI;
 class FUICommandList;
 class FStylusStateTracker;		// for stylus events
-class FModelingModeAssetAPI;
 
 UCLASS(Transient)
 class UModelingToolsEditorMode : public UBaseLegacyWidgetEdMode
@@ -64,7 +63,8 @@ protected:
 	virtual void OnToolStarted(UInteractiveToolManager* Manager, UInteractiveTool* Tool) override;
 	virtual void OnToolEnded(UInteractiveToolManager* Manager, UInteractiveTool* Tool) override;
 
-	TSharedPtr<FModelingModeAssetAPI> ModelingModeAssetGenerationAPI;
+	FDelegateHandle MeshCreatedEventHandle;
+	FDelegateHandle TextureCreatedEventHandle;
 
 	TUniquePtr<FStylusStateTracker> StylusStateTracker;
 
diff --git a/Engine/Plugins/Experimental/SampleToolsEditorMode/Source/Private/SampleTools/CreateActorSampleTool.cpp b/Engine/Plugins/Experimental/SampleToolsEditorMode/Source/Private/SampleTools/CreateActorSampleTool.cpp
--- a/Engine/Plugins/Experimental/SampleToolsEditorMode/Source/Private/SampleTools/CreateActorSampleTool.cpp
+++ b/Engine/Plugins/Experimental/SampleToolsEditorMode/Source/Private/SampleTools/CreateActorSampleTool.cpp
@@ -14,16 +14,10 @@
  */
 
 
-bool UCreateActorSampleToolBuilder::CanBuildTool(const FToolBuilderState & SceneState) const
-{
-	return (this->AssetAPI != nullptr);
-}
-
 UInteractiveTool* UCreateActorSampleToolBuilder::BuildTool(const FToolBuilderState & SceneState) const
 {
 	UCreateActorSampleTool* NewTool = NewObject<UCreateActorSampleTool>(SceneState.ToolManager);
 	NewTool->SetWorld(SceneState.World);
-	NewTool->SetAssetAPI(AssetAPI);
 	return NewTool;
 }
 
@@ -51,12 +45,6 @@ void UCreateActorSampleTool::SetWorld(UWorld* World)
 	this->TargetWorld = World;
 }
 
-void UCreateActorSampleTool::SetAssetAPI(IToolsContextAssetAPI* AssetAPIIn)
-{
-	this->AssetAPI = AssetAPIIn;
-}
-
-
 void UCreateActorSampleTool::Setup()
 {
 	USingleClickTool::Setup();
diff --git a/Engine/Plugins/Experimental/SampleToolsEditorMode/Source/Private/SampleTools/CreateActorSampleTool.h b/Engine/Plugins/Experimental/SampleToolsEditorMode/Source/Private/SampleTools/CreateActorSampleTool.h
--- a/Engine/Plugins/Experimental/SampleToolsEditorMode/Source/Private/SampleTools/CreateActorSampleTool.h
+++ b/Engine/Plugins/Experimental/SampleToolsEditorMode/Source/Private/SampleTools/CreateActorSampleTool.h
@@ -3,7 +3,6 @@
 #pragma once
 
 #include "CoreMinimal.h"
-#include "UObject/NoExportTypes.h"
 #include "InteractiveToolBuilder.h"
 #include "BaseTools/SingleClickTool.h"
 #include "CreateActorSampleTool.generated.h"
@@ -18,14 +17,7 @@ class SAMPLETOOLSEDITORMODE_API UCreateActorSampleToolBuilder : public UInteract
 	GENERATED_BODY()
 
 public:
-	IToolsContextAssetAPI* AssetAPI;
-
-	UCreateActorSampleToolBuilder() 
-	{
-		AssetAPI = nullptr;
-	}
-
-	virtual bool CanBuildTool(const FToolBuilderState& SceneState) const override;
+	virtual bool CanBuildTool(const FToolBuilderState& SceneState) const override { return true; }
 	virtual UInteractiveTool* BuildTool(const FToolBuilderState& SceneState) const override;
 };
 
@@ -68,7 +60,6 @@ public:
 	UCreateActorSampleTool();
 
 	virtual void SetWorld(UWorld* World);
-	virtual void SetAssetAPI(IToolsContextAssetAPI* AssetAPI);
 
 	virtual void Setup() override;
 
@@ -83,7 +74,4 @@ protected:
 protected:
 	/** target World we will raycast into and create Actor in */
 	UWorld* TargetWorld;
-
-	/** Access to the ToolContext's Asset Creation API. This is not currently used, but can be used to (eg) add Components, etc*/
-	IToolsContextAssetAPI* AssetAPI;
 };
\ No newline at end of file
diff --git a/Engine/Plugins/Experimental/SampleToolsEditorMode/Source/Private/SampleToolsEditorMode.cpp b/Engine/Plugins/Experimental/SampleToolsEditorMode/Source/Private/SampleToolsEditorMode.cpp
--- a/Engine/Plugins/Experimental/SampleToolsEditorMode/Source/Private/SampleToolsEditorMode.cpp
+++ b/Engine/Plugins/Experimental/SampleToolsEditorMode/Source/Private/SampleToolsEditorMode.cpp
@@ -56,7 +56,6 @@ void USampleToolsEditorMode::Enter()
 	const FSampleToolsEditorModeCommands& SampleToolCommands = FSampleToolsEditorModeCommands::Get();
 
 	auto CreateActorSampleToolBuilder = NewObject< UCreateActorSampleToolBuilder>(this);
-	CreateActorSampleToolBuilder->AssetAPI = ToolsContext->GetAssetAPI();
 	RegisterTool(SampleToolCommands.CreateActorTool, TEXT("CreateActorSampleTool"), CreateActorSampleToolBuilder);
 
 	RegisterTool(SampleToolCommands.DrawCurveOnMeshTool, TEXT("DrawCurveOnMeshSampleTool"), NewObject<UDrawCurveOnMeshSampleToolBuilder>(this));
