MeshModelingTools: UETOOL-3595 Migrate Tool Usage Analytics to UE5

diff --git a/Engine/Plugins/Experimental/ModelingToolsEditorMode/Source/ModelingToolsEditorMode/Private/ModelingToolsEditorMode.cpp b/Engine/Plugins/Experimental/ModelingToolsEditorMode/Source/ModelingToolsEditorMode/Private/ModelingToolsEditorMode.cpp
--- a/Engine/Plugins/Experimental/ModelingToolsEditorMode/Source/ModelingToolsEditorMode/Private/ModelingToolsEditorMode.cpp
+++ b/Engine/Plugins/Experimental/ModelingToolsEditorMode/Source/ModelingToolsEditorMode/Private/ModelingToolsEditorMode.cpp
@@ -120,6 +120,9 @@
 
 const FEditorModeID UModelingToolsEditorMode::EM_ModelingToolsEditorModeId = TEXT("EM_ModelingToolsEditorMode");
 
+FDateTime UModelingToolsEditorMode::LastModeStartTimestamp;
+FDateTime UModelingToolsEditorMode::LastToolStartTimestamp;
+
 UModelingToolsEditorMode::UModelingToolsEditorMode()
 {
 	Info = FEditorModeInfo(
@@ -641,11 +644,42 @@ void UModelingToolsEditorMode::Enter()
 	//
 	// Engine Analytics
 	//
-	// Log Analytic of mode starting
-	if( FEngineAnalytics::IsAvailable() )
+
+	// Log mode starting
+	if (FEngineAnalytics::IsAvailable())
 	{
-		FEngineAnalytics::GetProvider().RecordEvent( TEXT( "Editor.Usage.MeshModelingMode.Enter" ) );
+		UModelingToolsEditorMode::LastModeStartTimestamp = FDateTime::UtcNow();
+		TArray<FAnalyticsEventAttribute> Attributes;
+		Attributes.Add(FAnalyticsEventAttribute(TEXT("Timestamp"), LastModeStartTimestamp.ToString()));
+		FEngineAnalytics::GetProvider().RecordEvent(TEXT("Editor.Usage.MeshModelingMode.Enter"), Attributes);
 	}
+
+	// Log tool starting
+	ToolsContext->ToolManager->OnToolStarted.AddLambda([this](UInteractiveToolManager* Manager, UInteractiveTool* Tool)
+	{
+		if (FEngineAnalytics::IsAvailable())
+		{
+			UModelingToolsEditorMode::LastToolStartTimestamp = FDateTime::UtcNow();
+			TArray<FAnalyticsEventAttribute> Attributes;
+			Attributes.Add(FAnalyticsEventAttribute(TEXT("DisplayName"), Tool->GetToolInfo().ToolDisplayName.ToString()));
+			Attributes.Add(FAnalyticsEventAttribute(TEXT("Timestamp"), LastToolStartTimestamp.ToString()));
+			FEngineAnalytics::GetProvider().RecordEvent(TEXT("Editor.Usage.MeshModelingMode.ToolStarted"), Attributes);
+		}
+	});
+
+	// Log tool ending
+	ToolsContext->ToolManager->OnToolEnded.AddLambda([this](UInteractiveToolManager* Manager, UInteractiveTool* Tool)
+	{
+		if (FEngineAnalytics::IsAvailable())
+		{
+			TArray<FAnalyticsEventAttribute> Attributes;
+			Attributes.Add(FAnalyticsEventAttribute(TEXT("DisplayName"), Tool->GetToolInfo().ToolDisplayName.ToString()));
+			Attributes.Add(FAnalyticsEventAttribute(TEXT("Timestamp"), FDateTime::UtcNow().ToString()));
+			const FTimespan ToolUsageDuration = FDateTime::UtcNow() - UModelingToolsEditorMode::LastToolStartTimestamp;
+			Attributes.Add(FAnalyticsEventAttribute(TEXT("Duration.Seconds"), static_cast<float>(ToolUsageDuration.GetTotalSeconds())));
+			FEngineAnalytics::GetProvider().RecordEvent(TEXT("Editor.Usage.MeshModelingMode.ToolEnded"), Attributes);
+		}
+	});
 }
 
 void UModelingToolsEditorMode::Exit()
@@ -653,10 +687,14 @@ void UModelingToolsEditorMode::Exit()
 	//
 	// Engine Analytics
 	//
-	// Log Analytic of mode ending
-	if( FEngineAnalytics::IsAvailable() )
+	// Log mode ending
+	if (FEngineAnalytics::IsAvailable())
 	{
-		FEngineAnalytics::GetProvider().RecordEvent( TEXT( "Editor.Usage.MeshModelingMode.Exit" ) );
+		TArray<FAnalyticsEventAttribute> Attributes;
+		Attributes.Add(FAnalyticsEventAttribute(TEXT("Timestamp"), FDateTime::UtcNow().ToString()));
+		const FTimespan ModeUsageDuration = FDateTime::UtcNow() - UModelingToolsEditorMode::LastModeStartTimestamp;
+		Attributes.Add(FAnalyticsEventAttribute(TEXT("Duration.Seconds"), static_cast<float>(ModeUsageDuration.GetTotalSeconds())));
+		FEngineAnalytics::GetProvider().RecordEvent(TEXT("Editor.Usage.MeshModelingMode.Exit"), Attributes);
 	}
 
 	StylusStateTracker = nullptr;
diff --git a/Engine/Plugins/Experimental/ModelingToolsEditorMode/Source/ModelingToolsEditorMode/Public/ModelingToolsEditorMode.h b/Engine/Plugins/Experimental/ModelingToolsEditorMode/Source/ModelingToolsEditorMode/Public/ModelingToolsEditorMode.h
--- a/Engine/Plugins/Experimental/ModelingToolsEditorMode/Source/ModelingToolsEditorMode/Public/ModelingToolsEditorMode.h
+++ b/Engine/Plugins/Experimental/ModelingToolsEditorMode/Source/ModelingToolsEditorMode/Public/ModelingToolsEditorMode.h
@@ -71,4 +71,8 @@ protected:
 	void FocusCameraAtCursorHotkey();
 
 	void ConfigureRealTimeViewportsOverride(bool bEnable);
+
+	// analytics tracking
+	static FDateTime LastModeStartTimestamp;
+	static FDateTime LastToolStartTimestamp;
 };
