MeshModelingTools: Update the details panel in a more lightweight way when property values are changed programatically.

#rnx

diff --git a/Engine/Plugins/Experimental/ModelingToolsEditorMode/Source/ModelingToolsEditorMode/Private/ModelingToolsEditorModeToolkit.cpp b/Engine/Plugins/Experimental/ModelingToolsEditorMode/Source/ModelingToolsEditorMode/Private/ModelingToolsEditorModeToolkit.cpp
--- a/Engine/Plugins/Experimental/ModelingToolsEditorMode/Source/ModelingToolsEditorMode/Private/ModelingToolsEditorModeToolkit.cpp
+++ b/Engine/Plugins/Experimental/ModelingToolsEditorMode/Source/ModelingToolsEditorMode/Private/ModelingToolsEditorModeToolkit.cpp
@@ -389,9 +389,9 @@ void FModelingToolsEditorModeToolkit::UpdateActiveToolProperties()
 	}
 }
 
-void FModelingToolsEditorModeToolkit::RefreshPropertyObject(UObject* ChangedObject)
+void FModelingToolsEditorModeToolkit::InvalidateCachedDetailPanelState(UObject* ChangedObject)
 {
-	ModeDetailsView->ForceRefresh();
+	ModeDetailsView->InvalidateCachedState();
 }
 
 
@@ -722,7 +722,7 @@ void FModelingToolsEditorModeToolkit::OnToolStarted(UInteractiveToolManager* Man
 
 	UInteractiveTool* CurTool = GetScriptableEditorMode()->GetToolManager()->GetActiveTool(EToolSide::Left);
 	CurTool->OnPropertySetsModified.AddSP(this, &FModelingToolsEditorModeToolkit::UpdateActiveToolProperties);
-	CurTool->OnPropertyModifiedDirectlyByTool.AddSP(this, &FModelingToolsEditorModeToolkit::RefreshPropertyObject);
+	CurTool->OnPropertyModifiedDirectlyByTool.AddSP(this, &FModelingToolsEditorModeToolkit::InvalidateCachedDetailPanelState);
 
 	ModeHeaderArea->SetVisibility(EVisibility::Collapsed);
 	ActiveToolName = CurTool->GetToolInfo().ToolDisplayName;
diff --git a/Engine/Plugins/Experimental/ModelingToolsEditorMode/Source/ModelingToolsEditorMode/Public/ModelingToolsEditorModeToolkit.h b/Engine/Plugins/Experimental/ModelingToolsEditorMode/Source/ModelingToolsEditorMode/Public/ModelingToolsEditorModeToolkit.h
--- a/Engine/Plugins/Experimental/ModelingToolsEditorMode/Source/ModelingToolsEditorMode/Public/ModelingToolsEditorModeToolkit.h
+++ b/Engine/Plugins/Experimental/ModelingToolsEditorMode/Source/ModelingToolsEditorMode/Public/ModelingToolsEditorModeToolkit.h
@@ -74,7 +74,7 @@ private:
 
 	TSharedPtr<SWidget> ToolkitWidget;
 	void UpdateActiveToolProperties();
-	void RefreshPropertyObject(UObject* ChangedObject);
+	void InvalidateCachedDetailPanelState(UObject* ChangedObject);
 
 	TSharedPtr<SWidget> ViewportOverlayWidget;
 
diff --git a/Engine/Source/Editor/PropertyEditor/Private/SDetailsView.cpp b/Engine/Source/Editor/PropertyEditor/Private/SDetailsView.cpp
--- a/Engine/Source/Editor/PropertyEditor/Private/SDetailsView.cpp
+++ b/Engine/Source/Editor/PropertyEditor/Private/SDetailsView.cpp
@@ -911,6 +911,11 @@ void SDetailsView::SetOnObjectArrayChanged(FOnObjectArrayChanged OnObjectArrayCh
 }
 
 void SDetailsView::OnPostUndoRedo()
+{
+	InvalidateCachedState();
+}
+
+void SDetailsView::InvalidateCachedState()
 {
 	for (const TSharedPtr<FComplexPropertyNode>& RootNode : RootPropertyNodes)
 	{
diff --git a/Engine/Source/Editor/PropertyEditor/Private/SDetailsView.h b/Engine/Source/Editor/PropertyEditor/Private/SDetailsView.h
--- a/Engine/Source/Editor/PropertyEditor/Private/SDetailsView.h
+++ b/Engine/Source/Editor/PropertyEditor/Private/SDetailsView.h
@@ -29,6 +29,9 @@ public:
 	/** Causes the details view to be refreshed (new widgets generated) with the current set of objects */
 	virtual void ForceRefresh() override;
 
+	/** Invalidates cached state such as the "revert to default" arrow and edit conditions, without rebuilding the entire panel. */
+	virtual void InvalidateCachedState() override;
+
 	/** Move the scrolling offset (by item), but do not refresh the tree*/
 	void MoveScrollOffset(int32 DeltaOffset) override;
 
diff --git a/Engine/Source/Editor/PropertyEditor/Public/IDetailsView.h b/Engine/Source/Editor/PropertyEditor/Public/IDetailsView.h
--- a/Engine/Source/Editor/PropertyEditor/Public/IDetailsView.h
+++ b/Engine/Source/Editor/PropertyEditor/Public/IDetailsView.h
@@ -388,6 +388,12 @@ public:
 
 	/** Force refresh */
 	virtual void ForceRefresh() = 0;
+
+	/** Invalidates any cached state without necessarily doing a complete rebuild. */
+	virtual void InvalidateCachedState()
+	{
+		ForceRefresh();
+	}
 	
 	/** Sets an optional object filter to use for more complex handling of what a details panel is viewing. */
 	virtual void SetObjectFilter(TSharedPtr<FDetailsViewObjectFilter> InFilter) = 0;
