ModelingMode: move some toolkit UI initialization to after the mode is initialized, as it requires that to have ocurred first
#rnx

diff --git a/Engine/Plugins/Experimental/ModelingToolsEditorMode/Source/ModelingToolsEditorMode/Private/ModelingToolsEditorMode.cpp b/Engine/Plugins/Experimental/ModelingToolsEditorMode/Source/ModelingToolsEditorMode/Private/ModelingToolsEditorMode.cpp
--- a/Engine/Plugins/Experimental/ModelingToolsEditorMode/Source/ModelingToolsEditorMode/Private/ModelingToolsEditorMode.cpp
+++ b/Engine/Plugins/Experimental/ModelingToolsEditorMode/Source/ModelingToolsEditorMode/Private/ModelingToolsEditorMode.cpp
@@ -690,6 +690,14 @@ void UModelingToolsEditorMode::Enter()
 			FEngineAnalytics::GetProvider().RecordEvent(TEXT("Editor.Usage.MeshModelingMode.ToolEnded"), Attributes);
 		}
 	});
+
+	// do any toolkit UI initialization that depends on the mode setup above
+	if (Toolkit.IsValid())
+	{
+		FModelingToolsEditorModeToolkit* ModelingToolkit = (FModelingToolsEditorModeToolkit*)Toolkit.Get();
+		ModelingToolkit->InitializeAfterModeSetup();
+	}
+
 }
 
 void UModelingToolsEditorMode::Exit()
diff --git a/Engine/Plugins/Experimental/ModelingToolsEditorMode/Source/ModelingToolsEditorMode/Private/ModelingToolsEditorModeToolkit.cpp b/Engine/Plugins/Experimental/ModelingToolsEditorMode/Source/ModelingToolsEditorMode/Private/ModelingToolsEditorModeToolkit.cpp
--- a/Engine/Plugins/Experimental/ModelingToolsEditorMode/Source/ModelingToolsEditorMode/Private/ModelingToolsEditorModeToolkit.cpp
+++ b/Engine/Plugins/Experimental/ModelingToolsEditorMode/Source/ModelingToolsEditorMode/Private/ModelingToolsEditorModeToolkit.cpp
@@ -307,12 +307,6 @@ TSharedPtr<SWidget> FModelingToolsEditorModeToolkit::MakeAssetConfigPanel()
 
 	});
 
-
-	// initialize combos
-	//UpdateAssetPanelFromSettings();
-	AssetLODMode->SetSelectedItem(AssetLODModes[0]);		// TODO read from settings?
-
-
 	AssetLODModeLabel = SNew(STextBlock).Text(LOCTEXT("ActiveLODLabel", "Editing LOD"));
 
 	TSharedPtr<SVerticalBox> Content = SNew(SVerticalBox)
@@ -378,6 +372,19 @@ TSharedPtr<SWidget> FModelingToolsEditorModeToolkit::MakeAssetConfigPanel()
 }
 
 
+void FModelingToolsEditorModeToolkit::InitializeAfterModeSetup()
+{
+	if (bFirstInitializeAfterModeSetup)
+	{
+		// force update of the active asset LOD mode, this is necessary because the update modifies
+		// ToolTarget Factories that are only available once ModelingToolsEditorMode has been initialized
+		AssetLODMode->SetSelectedItem(AssetLODModes[0]);
+
+		bFirstInitializeAfterModeSetup = false;
+	}
+
+}
+
 
 
 void FModelingToolsEditorModeToolkit::UpdateActiveToolProperties()
diff --git a/Engine/Plugins/Experimental/ModelingToolsEditorMode/Source/ModelingToolsEditorMode/Public/ModelingToolsEditorModeToolkit.h b/Engine/Plugins/Experimental/ModelingToolsEditorMode/Source/ModelingToolsEditorMode/Public/ModelingToolsEditorModeToolkit.h
--- a/Engine/Plugins/Experimental/ModelingToolsEditorMode/Source/ModelingToolsEditorMode/Public/ModelingToolsEditorModeToolkit.h
+++ b/Engine/Plugins/Experimental/ModelingToolsEditorMode/Source/ModelingToolsEditorMode/Public/ModelingToolsEditorModeToolkit.h
@@ -31,6 +31,9 @@ public:
 	virtual FText GetBaseToolkitName() const override;
 	virtual TSharedPtr<class SWidget> GetInlineContent() const override { return ToolkitWidget; }
 
+	// initialize toolkit widgets that need to wait until mode is initialized/entered
+	virtual void InitializeAfterModeSetup();
+
 	// set/clear notification message area
 	virtual void PostNotification(const FText& Message);
 	virtual void ClearNotification();
@@ -106,6 +109,7 @@ private:
 	TSharedPtr<STextBlock> AssetLODModeLabel;
 	TSharedPtr<STextComboBox> AssetLODMode;
 
+	bool bFirstInitializeAfterModeSetup = true;
 
 	// custom accept/cancel/complete handlers
 };
