InteractiveToolsFramework: add IInteractiveToolNestedAcceptCancelAPI, to allow Tools to publish their ability to consume "sub-operation" accept or cancel interactions
ModelingMode: on Escape hotkey, if active Tool implements the above APIs, try to execute a nested-Cancel, and if it succeeds, do not cancel the Tool
ModelingTools: implement this api in MeshSelectionTool

#rnx

diff --git a/Engine/Plugins/Experimental/MeshModelingToolset/Source/MeshModelingTools/Private/MeshSelectionTool.cpp b/Engine/Plugins/Experimental/MeshModelingToolset/Source/MeshModelingTools/Private/MeshSelectionTool.cpp
--- a/Engine/Plugins/Experimental/MeshModelingToolset/Source/MeshModelingTools/Private/MeshSelectionTool.cpp
+++ b/Engine/Plugins/Experimental/MeshModelingToolset/Source/MeshModelingTools/Private/MeshSelectionTool.cpp
@@ -723,6 +723,27 @@ FBox UMeshSelectionTool::GetWorldSpaceFocusBox()
 
 
 
+bool UMeshSelectionTool::SupportsNestedCancelCommand()
+{
+	return true;
+}
+
+bool UMeshSelectionTool::CanCurrentlyNestedCancel()
+{
+	return Selection && (Selection->Faces.Num() > 0);
+}
+
+bool UMeshSelectionTool::ExecuteNestedCancelCommand()
+{
+	if (CanCurrentlyNestedCancel())
+	{
+		ClearSelection();
+		return true;
+	}
+	return false;
+}
+
+
 
 void UMeshSelectionTool::OnRegionHighlightUpdated()
 {
diff --git a/Engine/Plugins/Experimental/MeshModelingToolset/Source/MeshModelingTools/Public/MeshSelectionTool.h b/Engine/Plugins/Experimental/MeshModelingToolset/Source/MeshModelingTools/Public/MeshSelectionTool.h
--- a/Engine/Plugins/Experimental/MeshModelingToolset/Source/MeshModelingTools/Public/MeshSelectionTool.h
+++ b/Engine/Plugins/Experimental/MeshModelingToolset/Source/MeshModelingTools/Public/MeshSelectionTool.h
@@ -4,6 +4,7 @@
 
 #include "CoreMinimal.h"
 #include "DynamicMeshBrushTool.h"
+#include "InteractiveToolQueryInterfaces.h"
 #include "SelectionSet.h"
 #include "Changes/MeshSelectionChange.h"
 #include "DynamicMesh/DynamicMeshOctree3.h"
@@ -276,7 +277,7 @@ public:
  *
  */
 UCLASS()
-class MESHMODELINGTOOLS_API UMeshSelectionTool : public UDynamicMeshBrushTool
+class MESHMODELINGTOOLS_API UMeshSelectionTool : public UDynamicMeshBrushTool, public IInteractiveToolNestedAcceptCancelAPI
 {
 	GENERATED_BODY()
 
@@ -306,6 +307,11 @@ public:
 	// IInteractiveToolCameraFocusAPI implementation
 	virtual FBox GetWorldSpaceFocusBox() override;
 
+	// IInteractiveToolCancelAPI
+	virtual bool SupportsNestedCancelCommand() override;
+	virtual bool CanCurrentlyNestedCancel() override;
+	virtual bool ExecuteNestedCancelCommand() override;
+
 public:
 
 	virtual void RequestAction(EMeshSelectionToolActions ActionType);
diff --git a/Engine/Plugins/Experimental/ModelingToolsEditorMode/Source/ModelingToolsEditorMode/Private/ModelingToolsEditorMode.cpp b/Engine/Plugins/Experimental/ModelingToolsEditorMode/Source/ModelingToolsEditorMode/Private/ModelingToolsEditorMode.cpp
--- a/Engine/Plugins/Experimental/ModelingToolsEditorMode/Source/ModelingToolsEditorMode/Private/ModelingToolsEditorMode.cpp
+++ b/Engine/Plugins/Experimental/ModelingToolsEditorMode/Source/ModelingToolsEditorMode/Private/ModelingToolsEditorMode.cpp
@@ -800,22 +800,17 @@ void UModelingToolsEditorMode::BindCommands()
 	// These aren't activated by buttons but have default chords that bind the keypresses to the action.
 	CommandList->MapAction(
 		ToolManagerCommands.AcceptOrCompleteActiveTool,
-		FExecuteAction::CreateLambda([this]() {
-			const EToolShutdownType ShutdownType = ToolsContext->CanAcceptActiveTool() ? EToolShutdownType::Accept : EToolShutdownType::Completed;
-			ToolsContext->EndTool(ShutdownType);
-			}),
+		FExecuteAction::CreateLambda([this]() { AcceptActiveToolActionOrTool(); }),
 		FCanExecuteAction::CreateLambda([this]() {
 				return ToolsContext->CanAcceptActiveTool() || ToolsContext->CanCompleteActiveTool();
 			}),
 		FGetActionCheckState(),
 		FIsActionButtonVisible(),
 		EUIActionRepeatMode::RepeatDisabled);
+
 	CommandList->MapAction(
 		ToolManagerCommands.CancelOrCompleteActiveTool,
-		FExecuteAction::CreateLambda([this]() {
-			const EToolShutdownType ShutdownType = ToolsContext->CanCancelActiveTool() ? EToolShutdownType::Cancel : EToolShutdownType::Completed;
-			ToolsContext->EndTool(ShutdownType);
-			}),
+		FExecuteAction::CreateLambda([this]() { CancelActiveToolActionOrTool(); }),
 		FCanExecuteAction::CreateLambda([this]() {
 				return ToolsContext->CanCompleteActiveTool() || ToolsContext->CanCancelActiveTool();
 			}),
@@ -824,6 +819,50 @@ void UModelingToolsEditorMode::BindCommands()
 		EUIActionRepeatMode::RepeatDisabled);
 }
 
+
+void UModelingToolsEditorMode::AcceptActiveToolActionOrTool()
+{
+	// if we have an active Tool that implements 
+	if (GetToolManager()->HasAnyActiveTool())
+	{
+		UInteractiveTool* Tool = GetToolManager()->GetActiveTool(EToolSide::Mouse);
+		IInteractiveToolNestedAcceptCancelAPI* CancelAPI = Cast<IInteractiveToolNestedAcceptCancelAPI>(Tool);
+		if (CancelAPI && CancelAPI->SupportsNestedAcceptCommand() && CancelAPI->CanCurrentlyNestedAccept())
+		{
+			bool bAccepted = CancelAPI->ExecuteNestedAcceptCommand();
+			if (bAccepted)
+			{
+				return;
+			}
+		}
+	}
+
+	const EToolShutdownType ShutdownType = ToolsContext->CanAcceptActiveTool() ? EToolShutdownType::Accept : EToolShutdownType::Completed;
+	ToolsContext->EndTool(ShutdownType);
+}
+
+
+void UModelingToolsEditorMode::CancelActiveToolActionOrTool()
+{
+	// if we have an active Tool that implements 
+	if (GetToolManager()->HasAnyActiveTool())
+	{
+		UInteractiveTool* Tool = GetToolManager()->GetActiveTool(EToolSide::Mouse);
+		IInteractiveToolNestedAcceptCancelAPI* CancelAPI = Cast<IInteractiveToolNestedAcceptCancelAPI>(Tool);
+		if (CancelAPI && CancelAPI->SupportsNestedCancelCommand() && CancelAPI->CanCurrentlyNestedCancel())
+		{
+			bool bCancelled = CancelAPI->ExecuteNestedCancelCommand();
+			if (bCancelled)
+			{
+				return;
+			}
+		}
+	}
+
+	const EToolShutdownType ShutdownType = ToolsContext->CanCancelActiveTool() ? EToolShutdownType::Cancel : EToolShutdownType::Completed;
+	ToolsContext->EndTool(ShutdownType);
+}
+
 void UModelingToolsEditorMode::ModelingModeShortcutRequested(EModelingModeActionCommands Command)
 {
 	if (Command == EModelingModeActionCommands::FocusViewToCursor)
diff --git a/Engine/Plugins/Experimental/ModelingToolsEditorMode/Source/ModelingToolsEditorMode/Public/ModelingToolsEditorMode.h b/Engine/Plugins/Experimental/ModelingToolsEditorMode/Source/ModelingToolsEditorMode/Public/ModelingToolsEditorMode.h
--- a/Engine/Plugins/Experimental/ModelingToolsEditorMode/Source/ModelingToolsEditorMode/Public/ModelingToolsEditorMode.h
+++ b/Engine/Plugins/Experimental/ModelingToolsEditorMode/Source/ModelingToolsEditorMode/Public/ModelingToolsEditorMode.h
@@ -70,6 +70,9 @@ protected:
 	void ModelingModeShortcutRequested(EModelingModeActionCommands Command);
 	void FocusCameraAtCursorHotkey();
 
+	void AcceptActiveToolActionOrTool();
+	void CancelActiveToolActionOrTool();
+
 	void ConfigureRealTimeViewportsOverride(bool bEnable);
 
 	// analytics tracking
diff --git a/Engine/Source/Runtime/Experimental/InteractiveToolsFramework/Public/InteractiveToolQueryInterfaces.h b/Engine/Source/Runtime/Experimental/InteractiveToolsFramework/Public/InteractiveToolQueryInterfaces.h
--- a/Engine/Source/Runtime/Experimental/InteractiveToolsFramework/Public/InteractiveToolQueryInterfaces.h
+++ b/Engine/Source/Runtime/Experimental/InteractiveToolsFramework/Public/InteractiveToolQueryInterfaces.h
@@ -65,3 +65,66 @@ public:
 
 
 };
+
+
+
+
+
+
+
+// UInterface for IInteractiveToolNestedAcceptCancelAPI
+UINTERFACE(MinimalAPI)
+class UInteractiveToolNestedAcceptCancelAPI : public UInterface
+{
+	GENERATED_BODY()
+};
+
+/**
+ * IInteractiveToolNestedAcceptCancelAPI provides an API for a Tool to publish
+ * intent and ability to Accept or Cancel sub-operations. For example in a Tool
+ * that has an editable active Selection, we might want the Escape hotkey to
+ * Clear any active selection, and then on a second press, to Cancel the Tool. 
+ * This API allows a Tool to say "I can consume a Cancel action", and similarly
+ * for Accept (although this is much less common).
+ */
+class IInteractiveToolNestedAcceptCancelAPI
+{
+	GENERATED_BODY()
+public:
+
+	/**
+	 * @return true if the implementor of this API may be able to consume a Cancel action
+	 */
+	virtual bool SupportsNestedCancelCommand() { return false; }
+
+	/**
+	 * @return true if the implementor of this API can currently consume a Cancel action
+	 */
+	virtual bool CanCurrentlyNestedCancel() { return false; }
+
+	/**
+	 * Called by Editor levels to tell the implementor (eg Tool) to execute a nested Cancel action
+	 * @return true if the implementor consumed the Cancel action
+	 */
+	virtual bool ExecuteNestedCancelCommand() { return false; }
+
+
+
+	/**
+	 * @return true if the implementor of this API may be able to consume an Accept action
+	 */
+	virtual bool SupportsNestedAcceptCommand() { return false; }
+
+	/**
+	 * @return true if the implementor of this API can currently consume an Accept action
+	 */
+	virtual bool CanCurrentlyNestedAccept() { return false; }
+
+	/**
+	 * Called by Editor levels to tell the implementor (eg Tool) to execute a nested Accept action
+	 * @return true if the implementor consumed the Accept action
+	 */
+	virtual bool ExecuteNestedAcceptCommand() { return false; }
+
+
+};
\ No newline at end of file
