ModelingToolsEditorMode: Move plugin from Experimental to Editor

#rnx

diff --git a/Engine/Plugins/Editor/ModelingToolsEditorMode/ModelingToolsEditorMode.uplugin b/Engine/Plugins/Editor/ModelingToolsEditorMode/ModelingToolsEditorMode.uplugin
new file mode 100644
--- /dev/null
+++ b/Engine/Plugins/Editor/ModelingToolsEditorMode/ModelingToolsEditorMode.uplugin
@@ -0,0 +1,38 @@
+{
+	"FileVersion": 3,
+	"Version": 1,
+	"VersionName": "0.1",
+	"FriendlyName": "Modeling Tools Editor Mode",
+	"Description": "Modeling Tools Mode includes a suite of interactive tools for creating and editing meshes in the Editor",
+	"Category": "Editor",
+	"CreatedBy": "Epic Games, Inc.",
+	"CreatedByURL": "http://epicgames.com",
+	"DocsURL": "",
+	"MarketplaceURL": "",
+	"SupportURL": "",
+	"CanContainContent": true,
+	"IsBetaVersion": true,
+	"Installed": false,
+	"Modules": [
+		{
+			"Name": "ModelingToolsEditorMode",
+			"Type": "Editor",
+			"LoadingPhase": "Default",
+			"WhitelistPlatforms": [ "Win64" ]
+		}
+	],
+	"Plugins": [
+		{
+			"Name": "MeshModelingToolset",
+			"Enabled": true
+		},
+		{
+			"Name": "MeshLODToolset",
+			"Enabled": true
+		},
+		{
+			"Name": "StylusInput",
+			"Enabled": true
+		}
+	]
+}
\ No newline at end of file
diff --git a/Engine/Plugins/Editor/ModelingToolsEditorMode/Source/ModelingToolsEditorMode/ModelingToolsEditorMode.Build.cs b/Engine/Plugins/Editor/ModelingToolsEditorMode/Source/ModelingToolsEditorMode/ModelingToolsEditorMode.Build.cs
new file mode 100644
--- /dev/null
+++ b/Engine/Plugins/Editor/ModelingToolsEditorMode/Source/ModelingToolsEditorMode/ModelingToolsEditorMode.Build.cs
@@ -0,0 +1,80 @@
+// Copyright Epic Games, Inc. All Rights Reserved.
+
+using UnrealBuildTool;
+
+public class ModelingToolsEditorMode : ModuleRules
+{
+	public ModelingToolsEditorMode(ReadOnlyTargetRules Target) : base(Target)
+	{
+		PCHUsage = ModuleRules.PCHUsageMode.UseExplicitOrSharedPCHs;
+		
+		PublicIncludePaths.AddRange(
+			new string[] {
+                //"ContentBrowser"
+				// ... add public include paths required here ...
+			}
+			);
+				
+		
+		PrivateIncludePaths.AddRange(
+			new string[] {
+				// ... add other private include paths required here ...
+			}
+			);
+			
+		
+		PublicDependencyModuleNames.AddRange(
+			new string[]
+			{
+				"Core"
+				// ... add other public dependencies that you statically link with here ...
+			}
+			);
+			
+		
+		PrivateDependencyModuleNames.AddRange(
+			new string[]
+			{
+                "Core",
+                "CoreUObject",
+                "ApplicationCore",
+                "Slate",
+                "SlateCore",
+                "Engine",
+                "InputCore",
+				"EditorFramework",
+				"UnrealEd",
+                "ContentBrowser",
+                "LevelEditor",
+				"StatusBar",
+				"EditorStyle",
+                "Projects",
+				"MeshDescription",
+				"StaticMeshDescription",
+                "InteractiveToolsFramework",
+				"EditorInteractiveToolsFramework",
+				"GeometryFramework",
+				"ModelingComponents",
+				"ModelingComponentsEditorOnly",
+				"MeshModelingTools",
+				"MeshModelingToolsEditorOnly",
+				"MeshLODToolset",
+
+                "ViewportInteraction",
+
+				"StylusInput"
+
+				// ... add private dependencies that you statically link with here ...	
+			}
+            );
+		
+		
+		DynamicallyLoadedModuleNames.AddRange(
+			new string[]
+			{
+				"Settings"
+			}
+			);
+
+	}
+}
diff --git a/Engine/Plugins/Editor/ModelingToolsEditorMode/Source/ModelingToolsEditorMode/Private/DetailsCustomizations/ModelingToolsAxisFilterCustomization.cpp b/Engine/Plugins/Editor/ModelingToolsEditorMode/Source/ModelingToolsEditorMode/Private/DetailsCustomizations/ModelingToolsAxisFilterCustomization.cpp
new file mode 100644
--- /dev/null
+++ b/Engine/Plugins/Editor/ModelingToolsEditorMode/Source/ModelingToolsEditorMode/Private/DetailsCustomizations/ModelingToolsAxisFilterCustomization.cpp
@@ -0,0 +1,101 @@
+// Copyright Epic Games, Inc. All Rights Reserved.
+
+#include "DetailsCustomizations/ModelingToolPropertyCustomizations.h"
+
+#include "PropertyHandle.h"
+#include "IPropertyTypeCustomization.h"
+#include "IPropertyUtilities.h"
+#include "DetailWidgetRow.h"
+#include "EditorFontGlyphs.h"
+
+#include "PropertySets/AxisFilterPropertyType.h"
+
+
+#define LOCTEXT_NAMESPACE "ModelingToolsAxisFilterCustomization"
+
+
+TSharedRef<IPropertyTypeCustomization> FModelingToolsAxisFilterCustomization::MakeInstance()
+{
+	return MakeShareable(new FModelingToolsAxisFilterCustomization);
+}
+
+void FModelingToolsAxisFilterCustomization::CustomizeHeader(TSharedRef<IPropertyHandle> StructPropertyHandle, FDetailWidgetRow& HeaderRow, IPropertyTypeCustomizationUtils& StructCustomizationUtils)
+{
+	const float XYZPadding = 10.f;
+
+	TSharedPtr<IPropertyHandle> bAxisX = StructPropertyHandle->GetChildHandle(GET_MEMBER_NAME_CHECKED(FModelingToolsAxisFilter, bAxisX));
+	bAxisX->MarkHiddenByCustomization();
+
+	TSharedPtr<IPropertyHandle> bAxisY = StructPropertyHandle->GetChildHandle(GET_MEMBER_NAME_CHECKED(FModelingToolsAxisFilter, bAxisY));
+	bAxisY->MarkHiddenByCustomization();
+
+	TSharedPtr<IPropertyHandle> bAxisZ = StructPropertyHandle->GetChildHandle(GET_MEMBER_NAME_CHECKED(FModelingToolsAxisFilter, bAxisZ));
+	bAxisZ->MarkHiddenByCustomization();
+
+	HeaderRow.NameContent()
+		[
+			StructPropertyHandle->CreatePropertyNameWidget()
+		]
+	.ValueContent()
+		.MinDesiredWidth(125.f)
+		[
+			SNew(SHorizontalBox)
+			+ SHorizontalBox::Slot()
+			.Padding(0.f, 0.f, XYZPadding, 0.f)
+			.AutoWidth()
+			[
+				SNew(SHorizontalBox)
+				+ SHorizontalBox::Slot()
+				.AutoWidth()
+				[
+					bAxisX->CreatePropertyNameWidget()
+				]
+				+ SHorizontalBox::Slot()
+				.AutoWidth()
+				[
+					bAxisX->CreatePropertyValueWidget()
+				]
+			]
+
+			+ SHorizontalBox::Slot()
+			.Padding(0.f, 0.f, XYZPadding, 0.f)
+			.AutoWidth()
+			[
+				SNew(SHorizontalBox)
+				+ SHorizontalBox::Slot()
+				.AutoWidth()
+				[
+					bAxisY->CreatePropertyNameWidget()
+				]
+				+ SHorizontalBox::Slot()
+				.AutoWidth()
+				[
+					bAxisY->CreatePropertyValueWidget()
+				]
+			]
+
+			+ SHorizontalBox::Slot()
+			.Padding(0.f, 0.f, XYZPadding, 0.f)
+			.AutoWidth()
+			[
+				SNew(SHorizontalBox)
+				+ SHorizontalBox::Slot()
+				.AutoWidth()
+				[
+					bAxisZ->CreatePropertyNameWidget()
+				]
+				+ SHorizontalBox::Slot()
+				.AutoWidth()
+				[
+					bAxisZ->CreatePropertyValueWidget()
+				]
+			]
+		];
+}
+
+void FModelingToolsAxisFilterCustomization::CustomizeChildren(TSharedRef<IPropertyHandle> StructPropertyHandle, class IDetailChildrenBuilder& ChildBuilder, IPropertyTypeCustomizationUtils& StructCustomizationUtils)
+{
+}
+
+
+#undef LOCTEXT_NAMESPACE
\ No newline at end of file
diff --git a/Engine/Plugins/Editor/ModelingToolsEditorMode/Source/ModelingToolsEditorMode/Private/ModelingModeAssetUtils.cpp b/Engine/Plugins/Editor/ModelingToolsEditorMode/Source/ModelingToolsEditorMode/Private/ModelingModeAssetUtils.cpp
new file mode 100644
--- /dev/null
+++ b/Engine/Plugins/Editor/ModelingToolsEditorMode/Source/ModelingToolsEditorMode/Private/ModelingModeAssetUtils.cpp
@@ -0,0 +1,263 @@
+// Copyright Epic Games, Inc. All Rights Reserved.
+
+#include "ModelingModeAssetUtils.h"
+#include "ModelingToolsEditorModeSettings.h"
+
+#include "Engine/World.h"
+#include "AssetRegistryModule.h"
+#include "AssetToolsModule.h"
+#include "FileHelpers.h"
+#include "Misc/Paths.h"
+
+// for content-browser things
+#include "ContentBrowserModule.h"
+#include "IContentBrowserSingleton.h"
+
+
+#define LOCTEXT_NAMESPACE "ModelingModeAssetUtils"
+
+namespace UE
+{
+namespace Local
+{
+
+/**
+ * Generate a N-letter GUID string that contains only hex digits,
+ * and contains at least one letter and one number
+ */
+static FString GenerateRandomShortHexString(int32 NumChars = 8)
+{
+	int32 FailCount = 0;
+	while (FailCount++ < 10)
+	{
+		FGuid Guid = FGuid::NewGuid();
+		FString GuidString = Guid.ToString(EGuidFormats::UniqueObjectGuid).ToUpper();
+		FString Result;
+		int32 Digits = 0, Letters = 0;
+		for (int32 k = 0; k < GuidString.Len(); ++k)
+		{
+			TCHAR Character = GuidString[k];
+			if (FChar::IsHexDigit(Character))
+			{
+				Result.AppendChar(Character);
+				if (FChar::IsDigit(Character))
+				{
+					Digits++;
+				}
+				else
+				{
+					Letters++;
+				}
+			}
+			if (Result.Len() == NumChars)
+			{
+				if (Digits > 0 && Letters > 0)
+				{
+					return Result;
+				}
+				break;		// exit loop
+			}
+		}
+	}
+	return TEXT("BADGUID1");
+}
+
+
+
+FString GetWorldRelativeAssetRootPath(const UWorld* World)
+{
+	if (ensure(World->GetOutermost() != nullptr) == false)
+	{
+		return TEXT("/Game/");
+	}
+	FString WorldPackageName = World->GetOutermost()->GetName();
+	FString WorldPackageFolder = FPackageName::GetLongPackagePath(WorldPackageName);
+	return WorldPackageFolder;
+}
+
+
+FString GetActiveAssetFolderPath()
+{
+	IContentBrowserSingleton& ContentBrowser = FModuleManager::LoadModuleChecked<FContentBrowserModule>("ContentBrowser").Get();
+	return ContentBrowser.GetCurrentPath(EContentBrowserPathType::Internal);
+}
+
+
+FString MakeUniqueAssetName(const FString& FolderPath, const FString& AssetBaseName)
+{
+	FString UniquePackageName, UniqueAssetName;
+	FAssetToolsModule& AssetToolsModule = FModuleManager::LoadModuleChecked<FAssetToolsModule>("AssetTools");
+	AssetToolsModule.Get().CreateUniqueAssetName(
+		FolderPath + TEXT("/") + AssetBaseName, TEXT(""), UniquePackageName, UniqueAssetName);
+	return UniqueAssetName;
+}
+
+
+FString InteractiveSelectAssetPath(const FString& DefaultAssetName, const FText& DialogTitleMessage)
+{
+	IContentBrowserSingleton& ContentBrowser = FModuleManager::LoadModuleChecked<FContentBrowserModule>("ContentBrowser").Get();
+
+	FString UseDefaultAssetName = DefaultAssetName;
+	FString CurrentPath = GetActiveAssetFolderPath();
+	if (CurrentPath.IsEmpty() == false)
+	{
+		UseDefaultAssetName = MakeUniqueAssetName(CurrentPath, DefaultAssetName);
+	}
+
+	FSaveAssetDialogConfig Config;
+	Config.DefaultAssetName = UseDefaultAssetName;
+	Config.DialogTitleOverride = DialogTitleMessage;
+	Config.DefaultPath = CurrentPath;
+	return ContentBrowser.CreateModalSaveAssetDialog(Config);
+}
+
+
+
+}
+}
+
+
+
+
+
+
+FString UE::Modeling::GetNewAssetPathName(const FString& BaseNameIn, const UWorld* TargetWorld, FString SuggestedFolder)
+{
+	const UModelingToolsEditorModeSettings* Settings = GetDefault<UModelingToolsEditorModeSettings>();
+	EModelingModeAssetGenerationBehavior AutoGenMode = Settings->AssetGenerationMode;
+
+	FString PackageFolderPath = TEXT("/Game/");
+
+	if (SuggestedFolder.Len() > 0 && SuggestedFolder.StartsWith("/") )
+	{
+		// if we had a suggested folder, use it (todo: check that this folder exists? or can be created?)
+		PackageFolderPath = SuggestedFolder;
+	}
+	else
+	{
+		// Figure out the appropriate root path to use. If the level is in a plugin, this should
+		// be the plugin top-level path and not /Game.
+		if (Settings->AssetGenerationLocation == EModelingModeAssetGenerationLocation::AutoGeneratedWorldRelativeAssetPath)
+		{
+			PackageFolderPath = UE::Local::GetWorldRelativeAssetRootPath(TargetWorld);
+		}
+
+		// If the level is unsaved, then GetWorldRelativeAssetRootPath will return "/Temp". This is an undesirable
+		// place to even temporarily save assets because they must be manually moved elsewhere to save, and the /Temp folder 
+		// is somewhat troublesome to get to in the Editor. So if the flag below is set, then we will fall back to /Game
+		// instead. Projects with strict asset polices (due to P4/etc) may not want to allow this by default, so it is configurable.
+		if (PackageFolderPath.StartsWith("/Temp") && Settings->bStoreUnsavedLevelAssetsInTopLevelGameFolder)
+		{
+			PackageFolderPath = TEXT("/Game/");
+		}
+
+		// combine with fixed AutoGen path name if it is not empty
+		if (Settings->AutoGeneratedAssetPath.Len() > 0)
+		{
+			PackageFolderPath = FPaths::Combine(PackageFolderPath, Settings->AutoGeneratedAssetPath);
+		}
+
+		// append username-specific subfolder
+		if (Settings->bUsePerUserAutogenSubfolder)
+		{
+			FString UsernameString(FPlatformProcess::UserName());
+			if (UsernameString.Len() > 1)
+			{
+				PackageFolderPath = FPaths::Combine(PackageFolderPath, UsernameString);
+			}
+		}
+
+		// if we want to use the currently-visible asset browser path, try to find one (this can fail if no asset browser is visible/etc)
+		if (Settings->AssetGenerationLocation == EModelingModeAssetGenerationLocation::CurrentAssetBrowserPathIfAvailable)
+		{
+			FString CurrentAssetPath = UE::Local::GetActiveAssetFolderPath();
+			if (CurrentAssetPath.IsEmpty() == false)
+			{
+				PackageFolderPath = CurrentAssetPath;
+			}
+		}
+	}
+
+	FString ObjectBaseName = BaseNameIn;
+
+	// If we are in interactive mode, show the modal dialog and then get the path/name.
+	// If the user cancels, we are going to discard the asset
+	if (AutoGenMode == EModelingModeAssetGenerationBehavior::InteractivePromptToSave)
+	{
+		FString SelectedPath = UE::Local::InteractiveSelectAssetPath(ObjectBaseName, LOCTEXT("GenerateStaticMeshActorPathDialogWarning", "Choose Folder Path and Name for New Asset. Cancel to Discard New Asset."));
+		if (SelectedPath.IsEmpty() == false)
+		{
+			PackageFolderPath = FPaths::GetPath(SelectedPath);
+			ObjectBaseName = FPaths::GetBaseFilename(SelectedPath, true);
+		}
+		else
+		{
+			return FString();
+		}
+	}
+
+	FString UseBaseName = ObjectBaseName;
+	if (Settings->bAppendRandomStringToName)
+	{
+		FString GuidString = UE::Local::GenerateRandomShortHexString();
+		UseBaseName = FString::Printf(TEXT("%s_%s"), *UseBaseName, *GuidString);
+	}
+
+
+	FString PackageNameOut, AssetNameOut;
+	FAssetToolsModule& AssetToolsModule = FModuleManager::LoadModuleChecked<FAssetToolsModule>("AssetTools");
+	AssetToolsModule.Get().CreateUniqueAssetName(
+		FPaths::Combine(PackageFolderPath, UseBaseName), TEXT(""),
+		PackageNameOut, AssetNameOut);
+
+	return PackageNameOut;
+}
+
+
+
+bool UE::Modeling::AutoSaveAsset(UObject* Asset)
+{
+	UPackage* AssetPackage = Asset->GetPackage();
+	if (ensure(AssetPackage) == false)
+	{
+		return false;
+	}
+
+	Asset->MarkPackageDirty();
+	FAssetRegistryModule::AssetCreated(Asset);
+
+	TArray<UPackage*> PackagesToSave;
+	PackagesToSave.Add(AssetPackage);
+
+	bool bCheckDirty = true;
+	bool bPromptToSave = false;		// because we are autosaving
+	FEditorFileUtils::EPromptReturnCode Result = FEditorFileUtils::PromptForCheckoutAndSave(PackagesToSave, bCheckDirty, bPromptToSave);
+	ensure(Result == FEditorFileUtils::PR_Success);
+	return (Result == FEditorFileUtils::PR_Success);
+}
+
+
+
+
+
+void UE::Modeling::OnNewAssetCreated(UObject* Asset)
+{
+	const UModelingToolsEditorModeSettings* Settings = GetDefault<UModelingToolsEditorModeSettings>();
+	EModelingModeAssetGenerationBehavior AutoGenMode = Settings->AssetGenerationMode;
+
+	// Auto-save the new asset if the user has requested this behavior, or if an interactive
+	// save/path dialog was shown earlier in the process (InteractivePromptToSave mode)
+	if (AutoGenMode == EModelingModeAssetGenerationBehavior::AutoGenerateAndAutosave
+		 || AutoGenMode == EModelingModeAssetGenerationBehavior::InteractivePromptToSave)
+	{
+		AutoSaveAsset(Asset);
+	}
+	else if (AutoGenMode == EModelingModeAssetGenerationBehavior::AutoGenerateButDoNotAutosave)
+	{
+		Asset->MarkPackageDirty();
+		FAssetRegistryModule::AssetCreated(Asset);
+	}
+}
+
+
+#undef LOCTEXT_NAMESPACE
\ No newline at end of file
diff --git a/Engine/Plugins/Editor/ModelingToolsEditorMode/Source/ModelingToolsEditorMode/Private/ModelingToolsActions.cpp b/Engine/Plugins/Editor/ModelingToolsEditorMode/Source/ModelingToolsEditorMode/Private/ModelingToolsActions.cpp
new file mode 100644
--- /dev/null
+++ b/Engine/Plugins/Editor/ModelingToolsEditorMode/Source/ModelingToolsEditorMode/Private/ModelingToolsActions.cpp
@@ -0,0 +1,187 @@
+// Copyright Epic Games, Inc. All Rights Reserved.
+
+#include "ModelingToolsActions.h"
+#include "EditorStyleSet.h"
+#include "DynamicMeshSculptTool.h"
+#include "MeshVertexSculptTool.h"
+#include "MeshGroupPaintTool.h"
+#include "MeshInspectorTool.h"
+#include "DrawPolygonTool.h"
+#include "EditMeshPolygonsTool.h"
+#include "ShapeSprayTool.h"
+#include "MeshSpaceDeformerTool.h"
+#include "MeshSelectionTool.h"
+#include "TransformMeshesTool.h"
+#include "PlaneCutTool.h"
+#include "EditMeshPolygonsTool.h"
+#include "DrawAndRevolveTool.h"
+
+#define LOCTEXT_NAMESPACE "ModelingToolsCommands"
+
+
+
+FModelingModeActionCommands::FModelingModeActionCommands() :
+	TCommands<FModelingModeActionCommands>(
+		"ModelingModeCommands", // Context name for fast lookup
+		NSLOCTEXT("Contexts", "ModelingModeCommands", "Modeling Mode Shortcuts"), // Localized context name for displaying
+		NAME_None, // Parent
+		FEditorStyle::GetStyleSetName() // Icon Style Set
+		)
+{
+}
+
+
+void FModelingModeActionCommands::RegisterCommands()
+{
+	UI_COMMAND(FocusViewCommand, "Focus View at Cursor", "Focuses the camera at the scene hit location under the cursor", EUserInterfaceActionType::None, FInputChord(EKeys::C));
+}
+
+
+void FModelingModeActionCommands::RegisterCommandBindings(TSharedPtr<FUICommandList> UICommandList, TFunction<void(EModelingModeActionCommands)> OnCommandExecuted)
+{
+	const FModelingModeActionCommands& Commands = FModelingModeActionCommands::Get();
+
+	UICommandList->MapAction(
+		Commands.FocusViewCommand,
+		FExecuteAction::CreateLambda([OnCommandExecuted]() { OnCommandExecuted(EModelingModeActionCommands::FocusViewToCursor); }));
+}
+
+void FModelingModeActionCommands::UnRegisterCommandBindings(TSharedPtr<FUICommandList> UICommandList)
+{
+	const FModelingModeActionCommands& Commands = FModelingModeActionCommands::Get();
+	UICommandList->UnmapAction(Commands.FocusViewCommand);
+}
+
+
+
+FModelingToolActionCommands::FModelingToolActionCommands() : 
+	TInteractiveToolCommands<FModelingToolActionCommands>(
+		"ModelingToolsEditMode", // Context name for fast lookup
+		NSLOCTEXT("Contexts", "ModelingToolsEditMode", "Modeling Tools - Shared Shortcuts"), // Localized context name for displaying
+		NAME_None, // Parent
+		FEditorStyle::GetStyleSetName() // Icon Style Set
+	)
+{
+}
+
+
+void FModelingToolActionCommands::GetToolDefaultObjectList(TArray<UInteractiveTool*>& ToolCDOs)
+{
+	ToolCDOs.Add(GetMutableDefault<UMeshInspectorTool>());
+	//ToolCDOs.Add(GetMutableDefault<UDrawPolygonTool>());
+	//ToolCDOs.Add(GetMutableDefault<UEditMeshPolygonsTool>());
+	ToolCDOs.Add(GetMutableDefault<UMeshSpaceDeformerTool>());
+	ToolCDOs.Add(GetMutableDefault<UShapeSprayTool>());
+	//ToolCDOs.Add(GetMutableDefault<UPlaneCutTool>());
+}
+
+
+
+void FModelingToolActionCommands::RegisterAllToolActions()
+{
+	FModelingToolActionCommands::Register();
+	FSculptToolActionCommands::Register();
+	FVertexSculptToolActionCommands::Register();
+	FMeshGroupPaintToolActionCommands::Register();
+	FDrawPolygonToolActionCommands::Register();
+	FTransformToolActionCommands::Register();
+	FMeshSelectionToolActionCommands::Register();
+	FMeshPlaneCutToolActionCommands::Register();
+	FEditMeshPolygonsToolActionCommands::Register();
+	FDrawAndRevolveToolActionCommands::Register();
+}
+
+void FModelingToolActionCommands::UnregisterAllToolActions()
+{
+	FModelingToolActionCommands::Unregister();
+	FSculptToolActionCommands::Unregister();
+	FVertexSculptToolActionCommands::Unregister();
+	FMeshGroupPaintToolActionCommands::Unregister();
+	FDrawPolygonToolActionCommands::Unregister();
+	FTransformToolActionCommands::Unregister();
+	FMeshSelectionToolActionCommands::Unregister();
+	FMeshPlaneCutToolActionCommands::Unregister();
+	FEditMeshPolygonsToolActionCommands::Unregister();
+	FDrawAndRevolveToolActionCommands::Unregister();
+}
+
+
+
+void FModelingToolActionCommands::UpdateToolCommandBinding(UInteractiveTool* Tool, TSharedPtr<FUICommandList> UICommandList, bool bUnbind)
+{
+#define UPDATE_BINDING(CommandsType)  if (!bUnbind) CommandsType::Get().BindCommandsForCurrentTool(UICommandList, Tool); else CommandsType::Get().UnbindActiveCommands(UICommandList);
+
+
+	if (Cast<UTransformMeshesTool>(Tool) != nullptr)
+	{
+		UPDATE_BINDING(FTransformToolActionCommands);
+	}
+	else if (Cast<UDynamicMeshSculptTool>(Tool) != nullptr)
+	{
+		UPDATE_BINDING(FSculptToolActionCommands);
+	}
+	else if (Cast<UMeshVertexSculptTool>(Tool) != nullptr)
+	{
+		UPDATE_BINDING(FVertexSculptToolActionCommands);
+	}
+	else if (Cast<UMeshGroupPaintTool>(Tool) != nullptr)
+	{
+		UPDATE_BINDING(FMeshGroupPaintToolActionCommands);
+	}
+	else if (Cast<UDrawPolygonTool>(Tool) != nullptr)
+	{
+		UPDATE_BINDING(FDrawPolygonToolActionCommands);
+	}
+	else if (Cast<UMeshSelectionTool>(Tool) != nullptr)
+	{
+		UPDATE_BINDING(FMeshSelectionToolActionCommands);
+	}
+	else if (Cast<UPlaneCutTool>(Tool) != nullptr)
+	{
+		UPDATE_BINDING(FMeshPlaneCutToolActionCommands);
+	}
+	else if (Cast<UEditMeshPolygonsTool>(Tool) != nullptr)
+	{
+		UPDATE_BINDING(FEditMeshPolygonsToolActionCommands);
+	}
+	else if (Cast<UDrawAndRevolveTool>(Tool) != nullptr)
+	{
+		UPDATE_BINDING(FDrawAndRevolveToolActionCommands);
+	}
+	else
+	{
+		UPDATE_BINDING(FModelingToolActionCommands);
+	}
+}
+
+
+
+
+
+#define DEFINE_TOOL_ACTION_COMMANDS(CommandsClassName, ContextNameString, SettingsDialogString, ToolClassName ) \
+CommandsClassName::CommandsClassName() : TInteractiveToolCommands<CommandsClassName>( \
+ContextNameString, NSLOCTEXT("Contexts", ContextNameString, SettingsDialogString), NAME_None, FEditorStyle::GetStyleSetName()) {} \
+void CommandsClassName::GetToolDefaultObjectList(TArray<UInteractiveTool*>& ToolCDOs) \
+{\
+	ToolCDOs.Add(GetMutableDefault<ToolClassName>()); \
+}
+
+
+
+
+DEFINE_TOOL_ACTION_COMMANDS(FSculptToolActionCommands, "ModelingToolsSculptTool", "Modeling Tools - Sculpt Tool", UDynamicMeshSculptTool);
+DEFINE_TOOL_ACTION_COMMANDS(FVertexSculptToolActionCommands, "ModelingToolsVertexSculptTool", "Modeling Tools - Vertex Sculpt Tool", UMeshVertexSculptTool);
+DEFINE_TOOL_ACTION_COMMANDS(FMeshGroupPaintToolActionCommands, "ModelingToolsMeshGroupPaintTool", "Modeling Tools - Group Paint Tool", UMeshGroupPaintTool);
+DEFINE_TOOL_ACTION_COMMANDS(FTransformToolActionCommands, "ModelingToolsTransformTool", "Modeling Tools - Transform Tool", UTransformMeshesTool);
+DEFINE_TOOL_ACTION_COMMANDS(FDrawPolygonToolActionCommands, "ModelingToolsDrawPolygonTool", "Modeling Tools - Draw Polygon Tool", UDrawPolygonTool);
+DEFINE_TOOL_ACTION_COMMANDS(FMeshSelectionToolActionCommands, "ModelingToolsMeshSelectionTool", "Modeling Tools - Mesh Selection Tool", UMeshSelectionTool);
+DEFINE_TOOL_ACTION_COMMANDS(FMeshPlaneCutToolActionCommands, "ModelingToolsMeshPlaneCutTool", "Modeling Tools - Mesh Plane Cut Tool", UPlaneCutTool);
+DEFINE_TOOL_ACTION_COMMANDS(FEditMeshPolygonsToolActionCommands, "ModelingToolsEditMeshPolygonsTool", "Modeling Tools - Edit Mesh Polygons Tool", UEditMeshPolygonsTool);
+DEFINE_TOOL_ACTION_COMMANDS(FDrawAndRevolveToolActionCommands, "ModelingToolsDrawAndRevolveTool", "Modeling Tools - Draw-and-Revolve Tool", UDrawAndRevolveTool);
+
+
+
+
+
+#undef LOCTEXT_NAMESPACE
+
diff --git a/Engine/Plugins/Editor/ModelingToolsEditorMode/Source/ModelingToolsEditorMode/Private/ModelingToolsEditorMode.cpp b/Engine/Plugins/Editor/ModelingToolsEditorMode/Source/ModelingToolsEditorMode/Private/ModelingToolsEditorMode.cpp
new file mode 100644
--- /dev/null
+++ b/Engine/Plugins/Editor/ModelingToolsEditorMode/Source/ModelingToolsEditorMode/Private/ModelingToolsEditorMode.cpp
@@ -0,0 +1,986 @@
+// Copyright Epic Games, Inc. All Rights Reserved.
+
+#include "ModelingToolsEditorMode.h"
+#include "InteractiveTool.h"
+#include "InteractiveToolsSelectionStoreSubsystem.h"
+#include "ModelingToolsEditorModeToolkit.h"
+#include "Toolkits/ToolkitManager.h"
+#include "ToolTargetManager.h"
+#include "ContextObjectStore.h"
+#include "ToolTargets/StaticMeshComponentToolTarget.h"
+#include "ToolTargets/VolumeDynamicMeshToolTarget.h"
+#include "ToolTargets/VolumeMeshDescriptionToolTarget.h"
+#include "ToolTargets/DynamicMeshComponentToolTarget.h"
+#include "Framework/Commands/UICommandList.h"
+#include "Framework/Application/SlateApplication.h"
+#include "Framework/MultiBox/MultiBoxBuilder.h"
+#include "EditorViewportClient.h"
+#include "EngineAnalytics.h"
+#include "BaseGizmos/TransformGizmoUtil.h"
+
+#include "Features/IModularFeatures.h"
+#include "ModelingModeToolExtensions.h"
+
+//#include "SingleClickTool.h"
+//#include "MeshSurfacePointTool.h"
+//#include "MeshVertexDragTool.h"
+#include "DynamicMeshSculptTool.h"
+#include "MeshVertexSculptTool.h"
+#include "EditMeshPolygonsTool.h"
+#include "DeformMeshPolygonsTool.h"
+#include "SubdividePolyTool.h"
+#include "ConvertToPolygonsTool.h"
+#include "AddPrimitiveTool.h"
+#include "AddPatchTool.h"
+#include "RevolveBoundaryTool.h"
+#include "SmoothMeshTool.h"
+#include "OffsetMeshTool.h"
+#include "RemeshMeshTool.h"
+#include "SimplifyMeshTool.h"
+#include "MeshInspectorTool.h"
+#include "WeldMeshEdgesTool.h"
+#include "DrawPolygonTool.h"
+#include "DrawPolyPathTool.h"
+#include "DrawAndRevolveTool.h"
+#include "ShapeSprayTool.h"
+#include "MergeMeshesTool.h"
+#include "VoxelCSGMeshesTool.h"
+#include "VoxelSolidifyMeshesTool.h"
+#include "VoxelBlendMeshesTool.h"
+#include "VoxelMorphologyMeshesTool.h"
+#include "PlaneCutTool.h"
+#include "MirrorTool.h"
+#include "SelfUnionMeshesTool.h"
+#include "CSGMeshesTool.h"
+#include "CutMeshWithMeshTool.h"
+#include "BspConversionTool.h"
+#include "MeshToVolumeTool.h"
+#include "VolumeToMeshTool.h"
+#include "HoleFillTool.h"
+#include "PolygonOnMeshTool.h"
+#include "DisplaceMeshTool.h"
+#include "MeshSpaceDeformerTool.h"
+#include "EditNormalsTool.h"
+#include "RemoveOccludedTrianglesTool.h"
+#include "AttributeEditorTool.h"
+#include "TransformMeshesTool.h"
+#include "MeshSelectionTool.h"
+#include "UVProjectionTool.h"
+#include "UVLayoutTool.h"
+#include "EditMeshMaterialsTool.h"
+#include "EditPivotTool.h"
+#include "BakeTransformTool.h"
+#include "CombineMeshesTool.h"
+#include "AlignObjectsTool.h"
+#include "EditUVIslandsTool.h"
+#include "BakeMeshAttributeMapsTool.h"
+#include "BakeMeshAttributeVertexTool.h"
+#include "MeshAttributePaintTool.h"
+#include "ParameterizeMeshTool.h"
+#include "RecomputeUVsTool.h"
+#include "MeshTangentsTool.h"
+#include "ProjectToTargetTool.h"
+#include "LatticeDeformerTool.h"
+#include "SeamSculptTool.h"
+#include "MeshGroupPaintTool.h"
+#include "TransferMeshTool.h"
+#include "ConvertMeshesTool.h"
+#include "SplitMeshesTool.h"
+
+#include "Physics/PhysicsInspectorTool.h"
+#include "Physics/SetCollisionGeometryTool.h"
+#include "Physics/ExtractCollisionGeometryTool.h"
+//#include "Physics/EditCollisionGeometryTool.h"
+
+// hair tools
+#include "Hair/GroomToMeshTool.h"
+#include "Hair/GroomCardsEditorTool.h"
+#include "GenerateLODMeshesTool.h"
+
+// asset tools
+#include "Tools/GenerateStaticMeshLODAssetTool.h"
+#include "Tools/LODManagerTool.h"
+
+#include "EditorModeManager.h"
+
+// stylus support
+#include "IStylusInputModule.h"
+
+#include "LevelEditor.h"
+#include "IAssetViewport.h"
+#include "SLevelViewport.h"
+
+#include "ModelingToolsActions.h"
+#include "ModelingToolsManagerActions.h"
+#include "ModelingModeAssetUtils.h"
+#include "EditorModelingObjectsCreationAPI.h"
+
+#define LOCTEXT_NAMESPACE "UModelingToolsEditorMode"
+
+
+//#define ENABLE_DEBUG_PRINTING
+
+const FEditorModeID UModelingToolsEditorMode::EM_ModelingToolsEditorModeId = TEXT("EM_ModelingToolsEditorMode");
+
+FDateTime UModelingToolsEditorMode::LastModeStartTimestamp;
+FDateTime UModelingToolsEditorMode::LastToolStartTimestamp;
+
+UModelingToolsEditorMode::UModelingToolsEditorMode()
+{
+	Info = FEditorModeInfo(
+		EM_ModelingToolsEditorModeId,
+		LOCTEXT("ModelingToolsEditorModeName", "Modeling"),
+		FSlateIcon("ModelingToolsStyle", "LevelEditor.ModelingToolsMode", "LevelEditor.ModelingToolsMode.Small"),
+		true);
+}
+
+UModelingToolsEditorMode::UModelingToolsEditorMode(FVTableHelper& Helper)
+{
+}
+
+UModelingToolsEditorMode::~UModelingToolsEditorMode()
+{
+}
+
+bool UModelingToolsEditorMode::ProcessEditDelete()
+{
+	if (UEdMode::ProcessEditDelete())
+	{
+		return true;
+	}
+
+	// for now we disable deleting in an Accept-style tool because it can result in crashes if we are deleting target object
+	if ( GetToolManager()->HasAnyActiveTool() && GetToolManager()->GetActiveTool(EToolSide::Mouse)->HasAccept() )
+	{
+		GetToolManager()->DisplayMessage(
+			LOCTEXT("CannotDeleteWarning", "Cannot delete objects while this Tool is active"), EToolMessageLevel::UserWarning);
+		return true;
+	}
+
+	// If we didn't skip deletion, then we're probably deleting something, so it seems fair to clear
+	// the selection.
+	if (UInteractiveToolsSelectionStoreSubsystem* ToolSelectionStore = GEngine->GetEngineSubsystem<UInteractiveToolsSelectionStoreSubsystem>())
+	{
+		ToolSelectionStore->ClearStoredSelection();
+	}
+
+	return false;
+}
+
+
+bool UModelingToolsEditorMode::ProcessEditCut()
+{
+	// for now we disable deleting in an Accept-style tool because it can result in crashes if we are deleting target object
+	if (GetToolManager()->HasAnyActiveTool() && GetToolManager()->GetActiveTool(EToolSide::Mouse)->HasAccept())
+	{
+		GetToolManager()->DisplayMessage(
+			LOCTEXT("CannotCutWarning", "Cannot cut objects while this Tool is active"), EToolMessageLevel::UserWarning);
+		return true;
+	}
+
+	// If we're doing a cut, we should clear the tool selection.
+	if (UInteractiveToolsSelectionStoreSubsystem* ToolSelectionStore = GEngine->GetEngineSubsystem<UInteractiveToolsSelectionStoreSubsystem>())
+	{
+		ToolSelectionStore->ClearStoredSelection();
+	}
+
+	return false;
+}
+
+
+
+bool UModelingToolsEditorMode::CanAutoSave() const
+{
+	// prevent autosave if any tool is active
+	return ToolsContext->ToolManager->HasAnyActiveTool() == false;
+}
+
+bool UModelingToolsEditorMode::ShouldDrawWidget() const
+{ 
+	// allow standard xform gizmo if we don't have an active tool
+	if (ToolsContext != nullptr && ToolsContext->ToolManager->HasAnyActiveTool())
+	{
+		return false;
+	}
+
+	return UBaseLegacyWidgetEdMode::ShouldDrawWidget(); 
+}
+
+void UModelingToolsEditorMode::Tick(FEditorViewportClient* ViewportClient, float DeltaTime)
+{
+	Super::Tick(ViewportClient, DeltaTime);
+
+	if (Toolkit.IsValid())
+	{
+		FModelingToolsEditorModeToolkit* ModelingToolkit = (FModelingToolsEditorModeToolkit*)Toolkit.Get();
+		ModelingToolkit->EnableShowRealtimeWarning(ViewportClient->IsRealtime() == false);
+	}
+}
+
+//
+// FStylusStateTracker registers itself as a listener for stylus events and implements
+// the IToolStylusStateProviderAPI interface, which allows MeshSurfacePointTool implementations
+ // to query for the pen pressure.
+//
+// This is kind of a hack. Unfortunately the current Stylus module is a Plugin so it
+// cannot be used in the base ToolsFramework, and we need this in the Mode as a workaround.
+//
+class FStylusStateTracker : public IStylusMessageHandler, public IToolStylusStateProviderAPI
+{
+public:
+	const IStylusInputDevice* ActiveDevice = nullptr;
+	int32 ActiveDeviceIndex = -1;
+
+	bool bPenDown = false;
+	float ActivePressure = 1.0;
+
+	FStylusStateTracker()
+	{
+		UStylusInputSubsystem* StylusSubsystem = GEditor->GetEditorSubsystem<UStylusInputSubsystem>();
+		StylusSubsystem->AddMessageHandler(*this);
+
+		ActiveDevice = FindFirstPenDevice(StylusSubsystem, ActiveDeviceIndex);
+		bPenDown = false;
+	}
+
+	virtual ~FStylusStateTracker()
+	{
+		UStylusInputSubsystem* StylusSubsystem = GEditor->GetEditorSubsystem<UStylusInputSubsystem>();
+		StylusSubsystem->RemoveMessageHandler(*this);
+	}
+
+	virtual void OnStylusStateChanged(const FStylusState& NewState, int32 StylusIndex) override
+	{
+		if (ActiveDevice == nullptr)
+		{
+			UStylusInputSubsystem* StylusSubsystem = GEditor->GetEditorSubsystem<UStylusInputSubsystem>();
+			ActiveDevice = FindFirstPenDevice(StylusSubsystem, ActiveDeviceIndex);
+			bPenDown = false;
+		}
+		if (ActiveDevice != nullptr && ActiveDeviceIndex == StylusIndex)
+		{
+			bPenDown = NewState.IsStylusDown();
+			ActivePressure = NewState.GetPressure();
+		}
+	}
+
+
+	bool HaveActiveStylusState() const
+	{
+		return ActiveDevice != nullptr && bPenDown;
+	}
+
+	static const IStylusInputDevice* FindFirstPenDevice(const UStylusInputSubsystem* StylusSubsystem, int32& ActiveDeviceOut)
+	{
+		int32 NumDevices = StylusSubsystem->NumInputDevices();
+		for (int32 k = 0; k < NumDevices; ++k)
+		{
+			const IStylusInputDevice* Device = StylusSubsystem->GetInputDevice(k);
+			const TArray<EStylusInputType>& Inputs = Device->GetSupportedInputs();
+			for (EStylusInputType Input : Inputs)
+			{
+				if (Input == EStylusInputType::Pressure)
+				{
+					ActiveDeviceOut = k;
+					return Device;
+				}
+			}
+		}
+		return nullptr;
+	}
+
+
+
+	// IToolStylusStateProviderAPI implementation
+	virtual float GetCurrentPressure() const override
+	{
+		return (ActiveDevice != nullptr && bPenDown) ? ActivePressure : 1.0f;
+	}
+
+};
+
+void UModelingToolsEditorMode::Enter()
+{
+	UEdMode::Enter();
+
+	// Register builders for tool targets that the mode uses.
+	ToolsContext->TargetManager->AddTargetFactory(NewObject<UStaticMeshComponentToolTargetFactory>(ToolsContext->TargetManager));
+	ToolsContext->TargetManager->AddTargetFactory(NewObject<UVolumeDynamicMeshToolTargetFactory>(ToolsContext->TargetManager));
+	ToolsContext->TargetManager->AddTargetFactory(NewObject<UVolumeMeshDescriptionToolTargetFactory>(ToolsContext->TargetManager));
+	ToolsContext->TargetManager->AddTargetFactory(NewObject<UDynamicMeshComponentToolTargetFactory>(ToolsContext->TargetManager));
+
+	// register stylus event handler
+	StylusStateTracker = MakeUnique<FStylusStateTracker>();
+
+	// register gizmo helper
+	UE::TransformGizmoUtil::RegisterTransformGizmoContextObject(ToolsContext.Get());
+
+	// disable HitProxy rendering, it is not used in Modeling Mode and adds overhead to Render() calls
+	ToolsContext.Get()->SetEnableRenderingDuringHitProxyPass(false);
+
+	// register object creation api
+	UEditorModelingObjectsCreationAPI* ModelCreationAPI = UEditorModelingObjectsCreationAPI::Register(ToolsContext.Get());
+	if (ModelCreationAPI)
+	{
+		ModelCreationAPI->GetNewAssetPathNameCallback.BindLambda([](const FString& BaseName, const UWorld* TargetWorld, FString SuggestedFolder)
+		{
+			return UE::Modeling::GetNewAssetPathName(BaseName, TargetWorld, SuggestedFolder);
+		});
+		MeshCreatedEventHandle = ModelCreationAPI->OnModelingMeshCreated.AddLambda([](const FCreateMeshObjectResult& CreatedInfo) 
+		{
+			if (CreatedInfo.NewAsset != nullptr)
+			{
+				UE::Modeling::OnNewAssetCreated(CreatedInfo.NewAsset);
+			}
+		});
+		TextureCreatedEventHandle = ModelCreationAPI->OnModelingTextureCreated.AddLambda([](const FCreateTextureObjectResult& CreatedInfo)
+		{
+			if (CreatedInfo.NewAsset != nullptr)
+			{
+				UE::Modeling::OnNewAssetCreated(CreatedInfo.NewAsset);
+			}
+		});
+	}
+
+	const FModelingToolsManagerCommands& ToolManagerCommands = FModelingToolsManagerCommands::Get();
+
+	// register tool set
+
+	//
+	// primitive tools
+	//
+	auto RegisterPrimitiveToolFunc  =
+		[this](TSharedPtr<FUICommandInfo> UICommand,
+								  FString&& ToolIdentifier,
+								  UAddPrimitiveToolBuilder::EMakeMeshShapeType ShapeTypeIn)
+	{
+		auto AddPrimitiveToolBuilder = NewObject<UAddPrimitiveToolBuilder>();
+		AddPrimitiveToolBuilder->ShapeType = ShapeTypeIn;
+		RegisterTool(UICommand, ToolIdentifier, AddPrimitiveToolBuilder);
+	};
+	RegisterPrimitiveToolFunc(ToolManagerCommands.BeginAddBoxPrimitiveTool,
+							  TEXT("BeginAddBoxPrimitiveTool"),
+							  UAddPrimitiveToolBuilder::EMakeMeshShapeType::Box);
+	RegisterPrimitiveToolFunc(ToolManagerCommands.BeginAddCylinderPrimitiveTool,
+							  TEXT("BeginAddCylinderPrimitiveTool"),
+							  UAddPrimitiveToolBuilder::EMakeMeshShapeType::Cylinder);
+	RegisterPrimitiveToolFunc(ToolManagerCommands.BeginAddConePrimitiveTool,
+							  TEXT("BeginAddConePrimitiveTool"),
+							  UAddPrimitiveToolBuilder::EMakeMeshShapeType::Cone);
+	RegisterPrimitiveToolFunc(ToolManagerCommands.BeginAddArrowPrimitiveTool,
+							  TEXT("BeginAddArrowPrimitiveTool"),
+							  UAddPrimitiveToolBuilder::EMakeMeshShapeType::Arrow);
+	RegisterPrimitiveToolFunc(ToolManagerCommands.BeginAddRectanglePrimitiveTool,
+							  TEXT("BeginAddRectanglePrimitiveTool"),
+							  UAddPrimitiveToolBuilder::EMakeMeshShapeType::Rectangle);
+	RegisterPrimitiveToolFunc(ToolManagerCommands.BeginAddDiscPrimitiveTool,
+							  TEXT("BeginAddDiscPrimitiveTool"),
+							  UAddPrimitiveToolBuilder::EMakeMeshShapeType::Disc);
+	RegisterPrimitiveToolFunc(ToolManagerCommands.BeginAddTorusPrimitiveTool,
+							  TEXT("BeginAddTorusPrimitiveTool"),
+							  UAddPrimitiveToolBuilder::EMakeMeshShapeType::Torus);
+	RegisterPrimitiveToolFunc(ToolManagerCommands.BeginAddSpherePrimitiveTool,
+							  TEXT("BeginAddSpherePrimitiveTool"),
+							  UAddPrimitiveToolBuilder::EMakeMeshShapeType::Sphere);
+	RegisterPrimitiveToolFunc(ToolManagerCommands.BeginAddStairsPrimitiveTool,
+							  TEXT("BeginAddStairsPrimitiveTool"),
+							  UAddPrimitiveToolBuilder::EMakeMeshShapeType::Stairs);
+
+	//
+	// make shape tools
+	//
+	auto AddPatchToolBuilder = NewObject<UAddPatchToolBuilder>();
+	RegisterTool(ToolManagerCommands.BeginAddPatchTool, TEXT("BeginAddPatchTool"), AddPatchToolBuilder);
+
+	auto RevolveBoundaryToolBuilder = NewObject<URevolveBoundaryToolBuilder>();
+	RegisterTool(ToolManagerCommands.BeginRevolveBoundaryTool, TEXT("BeginRevolveBoundaryTool"), RevolveBoundaryToolBuilder);
+
+	auto DrawPolygonToolBuilder = NewObject<UDrawPolygonToolBuilder>();
+	RegisterTool(ToolManagerCommands.BeginDrawPolygonTool, TEXT("BeginDrawPolygonTool"), DrawPolygonToolBuilder);
+
+	auto DrawPolyPathToolBuilder = NewObject<UDrawPolyPathToolBuilder>();
+	RegisterTool(ToolManagerCommands.BeginDrawPolyPathTool, TEXT("BeginDrawPolyPathTool"), DrawPolyPathToolBuilder);
+
+	auto DrawAndRevolveToolBuilder = NewObject<UDrawAndRevolveToolBuilder>();
+	RegisterTool(ToolManagerCommands.BeginDrawAndRevolveTool, TEXT("BeginDrawAndRevolveTool"), DrawAndRevolveToolBuilder);
+
+	auto ShapeSprayToolBuilder = NewObject<UShapeSprayToolBuilder>();
+	RegisterTool(ToolManagerCommands.BeginShapeSprayTool, TEXT("BeginShapeSprayTool"), ShapeSprayToolBuilder);
+
+
+	//
+	// vertex deform tools
+	//
+
+	auto MoveVerticesToolBuilder = NewObject<UMeshVertexSculptToolBuilder>();
+	MoveVerticesToolBuilder->StylusAPI = StylusStateTracker.Get();
+	RegisterTool(ToolManagerCommands.BeginSculptMeshTool, TEXT("BeginSculptMeshTool"), MoveVerticesToolBuilder);
+
+	auto MeshGroupPaintToolBuilder = NewObject<UMeshGroupPaintToolBuilder>();
+	MeshGroupPaintToolBuilder->StylusAPI = StylusStateTracker.Get();
+	RegisterTool(ToolManagerCommands.BeginMeshGroupPaintTool, TEXT("BeginMeshGroupPaintTool"), MeshGroupPaintToolBuilder);
+
+	RegisterTool(ToolManagerCommands.BeginPolyEditTool, TEXT("BeginPolyEditTool"), NewObject<UEditMeshPolygonsToolBuilder>());
+	UEditMeshPolygonsToolBuilder* TriEditBuilder = NewObject<UEditMeshPolygonsToolBuilder>();
+	TriEditBuilder->bTriangleMode = true;
+	RegisterTool(ToolManagerCommands.BeginTriEditTool, TEXT("BeginTriEditTool"), TriEditBuilder);
+	RegisterTool(ToolManagerCommands.BeginPolyDeformTool, TEXT("BeginPolyDeformTool"), NewObject<UDeformMeshPolygonsToolBuilder>());
+	RegisterTool(ToolManagerCommands.BeginSmoothMeshTool, TEXT("BeginSmoothMeshTool"), NewObject<USmoothMeshToolBuilder>());
+	RegisterTool(ToolManagerCommands.BeginOffsetMeshTool, TEXT("BeginOffsetMeshTool"), NewObject<UOffsetMeshToolBuilder>());
+	RegisterTool(ToolManagerCommands.BeginDisplaceMeshTool, TEXT("BeginDisplaceMeshTool"), NewObject<UDisplaceMeshToolBuilder>());
+	RegisterTool(ToolManagerCommands.BeginMeshSpaceDeformerTool, TEXT("BeginMeshSpaceDeformerTool"), NewObject<UMeshSpaceDeformerToolBuilder>());
+	RegisterTool(ToolManagerCommands.BeginTransformMeshesTool, TEXT("BeginTransformMeshesTool"), NewObject<UTransformMeshesToolBuilder>());
+	RegisterTool(ToolManagerCommands.BeginEditPivotTool, TEXT("BeginEditPivotTool"), NewObject<UEditPivotToolBuilder>());
+	RegisterTool(ToolManagerCommands.BeginAlignObjectsTool, TEXT("BeginAlignObjectsTool"), NewObject<UAlignObjectsToolBuilder>());
+	RegisterTool(ToolManagerCommands.BeginTransferMeshTool, TEXT("BeginTransferMeshTool"), NewObject<UTransferMeshToolBuilder>());
+	RegisterTool(ToolManagerCommands.BeginConvertMeshesTool, TEXT("BeginConvertMeshesTool"), NewObject<UConvertMeshesToolBuilder>());
+	RegisterTool(ToolManagerCommands.BeginSplitMeshesTool, TEXT("BeginSplitMeshesTool"), NewObject<USplitMeshesToolBuilder>());
+	RegisterTool(ToolManagerCommands.BeginBakeTransformTool, TEXT("BeginBakeTransformTool"), NewObject<UBakeTransformToolBuilder>());
+	RegisterTool(ToolManagerCommands.BeginTransformUVIslandsTool, TEXT("BeginTransformUVIslandsTool"), NewObject<UEditUVIslandsToolBuilder>());
+	RegisterTool(ToolManagerCommands.BeginLatticeDeformerTool, TEXT("BeginLatticeDeformerTool"), NewObject<ULatticeDeformerToolBuilder>());
+	RegisterTool(ToolManagerCommands.BeginSubdividePolyTool, TEXT("BeginSubdividePolyTool"), NewObject<USubdividePolyToolBuilder>());
+
+	UCombineMeshesToolBuilder* CombineMeshesToolBuilder = NewObject<UCombineMeshesToolBuilder>();
+	RegisterTool(ToolManagerCommands.BeginCombineMeshesTool, TEXT("BeginCombineMeshesTool"), CombineMeshesToolBuilder);
+
+	UCombineMeshesToolBuilder* DuplicateMeshesToolBuilder = NewObject<UCombineMeshesToolBuilder>();
+	DuplicateMeshesToolBuilder->bIsDuplicateTool = true;
+	RegisterTool(ToolManagerCommands.BeginDuplicateMeshesTool, TEXT("BeginDuplicateMeshesTool"), DuplicateMeshesToolBuilder);
+
+
+	ULODManagerToolBuilder* LODManagerToolBuilder = NewObject<ULODManagerToolBuilder>();
+	RegisterTool(ToolManagerCommands.BeginLODManagerTool, TEXT("BeginLODManagerTool"), LODManagerToolBuilder);
+
+	UGenerateStaticMeshLODAssetToolBuilder* GenerateSMLODToolBuilder = NewObject<UGenerateStaticMeshLODAssetToolBuilder>();
+	RegisterTool(ToolManagerCommands.BeginGenerateStaticMeshLODAssetTool, TEXT("BeginGenerateStaticMeshLODAssetTool"), GenerateSMLODToolBuilder);
+
+
+	// edit tools
+
+
+	auto DynaSculptToolBuilder = NewObject<UDynamicMeshSculptToolBuilder>();
+	DynaSculptToolBuilder->bEnableRemeshing = true;
+	DynaSculptToolBuilder->StylusAPI = StylusStateTracker.Get();
+	RegisterTool(ToolManagerCommands.BeginRemeshSculptMeshTool, TEXT("BeginRemeshSculptMeshTool"), DynaSculptToolBuilder);
+
+	RegisterTool(ToolManagerCommands.BeginRemeshMeshTool, TEXT("BeginRemeshMeshTool"), NewObject<URemeshMeshToolBuilder>());
+	RegisterTool(ToolManagerCommands.BeginProjectToTargetTool, TEXT("BeginProjectToTargetTool"), NewObject<UProjectToTargetToolBuilder>());
+	RegisterTool(ToolManagerCommands.BeginSimplifyMeshTool, TEXT("BeginSimplifyMeshTool"), NewObject<USimplifyMeshToolBuilder>());
+
+	auto EditNormalsToolBuilder = NewObject<UEditNormalsToolBuilder>();
+	RegisterTool(ToolManagerCommands.BeginEditNormalsTool, TEXT("BeginEditNormalsTool"), EditNormalsToolBuilder);
+
+	RegisterTool(ToolManagerCommands.BeginEditTangentsTool, TEXT("BeginEditTangentsTool"), NewObject<UMeshTangentsToolBuilder>());
+
+	auto RemoveOccludedTrianglesToolBuilder = NewObject<URemoveOccludedTrianglesToolBuilder>();
+	RegisterTool(ToolManagerCommands.BeginRemoveOccludedTrianglesTool, TEXT("BeginRemoveOccludedTrianglesTool"), RemoveOccludedTrianglesToolBuilder);
+
+	RegisterTool(ToolManagerCommands.BeginHoleFillTool, TEXT("BeginHoleFillTool"), NewObject<UHoleFillToolBuilder>());
+
+	auto UVProjectionToolBuilder = NewObject<UUVProjectionToolBuilder>();
+	RegisterTool(ToolManagerCommands.BeginUVProjectionTool, TEXT("BeginUVProjectionTool"), UVProjectionToolBuilder);
+
+	auto UVLayoutToolBuilder = NewObject<UUVLayoutToolBuilder>();
+	RegisterTool(ToolManagerCommands.BeginUVLayoutTool, TEXT("BeginUVLayoutTool"), UVLayoutToolBuilder);
+
+	auto MergeMeshesToolBuilder = NewObject<UMergeMeshesToolBuilder>();
+	RegisterTool(ToolManagerCommands.BeginVoxelMergeTool, TEXT("BeginVoxelMergeTool"), MergeMeshesToolBuilder);
+
+	auto VoxelCSGMeshesToolBuilder = NewObject<UVoxelCSGMeshesToolBuilder>();
+	RegisterTool(ToolManagerCommands.BeginVoxelBooleanTool, TEXT("BeginVoxelBooleanTool"), VoxelCSGMeshesToolBuilder);
+
+	auto VoxelSolidifyMeshesToolBuilder = NewObject<UVoxelSolidifyMeshesToolBuilder>();
+	RegisterTool(ToolManagerCommands.BeginVoxelSolidifyTool, TEXT("BeginVoxelSolidifyTool"), VoxelSolidifyMeshesToolBuilder);
+
+	auto VoxelBlendMeshesToolBuilder = NewObject<UVoxelBlendMeshesToolBuilder>();
+	RegisterTool(ToolManagerCommands.BeginVoxelBlendTool, TEXT("BeginVoxelBlendTool"), VoxelBlendMeshesToolBuilder);
+
+	auto VoxelMorphologyMeshesToolBuilder = NewObject<UVoxelMorphologyMeshesToolBuilder>();
+	RegisterTool(ToolManagerCommands.BeginVoxelMorphologyTool, TEXT("BeginVoxelMorphologyTool"), VoxelMorphologyMeshesToolBuilder);
+
+	auto SelfUnionMeshesToolBuilder = NewObject<USelfUnionMeshesToolBuilder>();
+	RegisterTool(ToolManagerCommands.BeginSelfUnionTool, TEXT("BeginSelfUnionTool"), SelfUnionMeshesToolBuilder);
+
+	auto CSGMeshesToolBuilder = NewObject<UCSGMeshesToolBuilder>();
+	RegisterTool(ToolManagerCommands.BeginMeshBooleanTool, TEXT("BeginMeshBooleanTool"), CSGMeshesToolBuilder);
+
+	auto CutMeshWithMeshToolBuilder = NewObject<UCutMeshWithMeshToolBuilder>();
+	RegisterTool(ToolManagerCommands.BeginCutMeshWithMeshTool, TEXT("BeginCutMeshWithMeshTool"), CutMeshWithMeshToolBuilder);
+
+	auto TrimMeshesToolBuilder = NewObject<UCSGMeshesToolBuilder>();
+	TrimMeshesToolBuilder->bTrimMode = true;
+	RegisterTool(ToolManagerCommands.BeginMeshTrimTool, TEXT("BeginMeshTrimTool"), TrimMeshesToolBuilder);
+
+	RegisterTool(ToolManagerCommands.BeginBspConversionTool, TEXT("BeginBspConversionTool"), NewObject<UBspConversionToolBuilder>());
+
+	auto MeshToVolumeToolBuilder = NewObject<UMeshToVolumeToolBuilder>();
+	RegisterTool(ToolManagerCommands.BeginMeshToVolumeTool, TEXT("BeginMeshToVolumeTool"), MeshToVolumeToolBuilder);
+
+	auto VolumeToMeshToolBuilder = NewObject<UVolumeToMeshToolBuilder>();
+	RegisterTool(ToolManagerCommands.BeginVolumeToMeshTool, TEXT("BeginVolumeToMeshTool"), VolumeToMeshToolBuilder);
+
+	auto PlaneCutToolBuilder = NewObject<UPlaneCutToolBuilder>();
+	RegisterTool(ToolManagerCommands.BeginPlaneCutTool, TEXT("BeginPlaneCutTool"), PlaneCutToolBuilder);
+
+	auto MirrorToolBuilder = NewObject<UMirrorToolBuilder>();
+	RegisterTool(ToolManagerCommands.BeginMirrorTool, TEXT("BeginMirrorTool"), MirrorToolBuilder);
+
+	auto PolygonCutToolBuilder = NewObject<UPolygonOnMeshToolBuilder>();
+	RegisterTool(ToolManagerCommands.BeginPolygonCutTool, TEXT("BeginPolygonCutTool"), PolygonCutToolBuilder);
+
+	auto GlobalUVGenerateToolBuilder = NewObject<UParameterizeMeshToolBuilder>();
+	RegisterTool(ToolManagerCommands.BeginGlobalUVGenerateTool, TEXT("BeginGlobalUVGenerateTool"), GlobalUVGenerateToolBuilder);
+
+	auto RecomputeUVsToolBuilder = NewObject<URecomputeUVsToolBuilder>();
+	RegisterTool(ToolManagerCommands.BeginGroupUVGenerateTool, TEXT("BeginGroupUVGenerateTool"), RecomputeUVsToolBuilder);
+
+	RegisterTool(ToolManagerCommands.BeginUVSeamEditTool, TEXT("BeginUVSeamEditTool"), NewObject< USeamSculptToolBuilder>());
+
+	auto MeshSelectionToolBuilder = NewObject<UMeshSelectionToolBuilder>();
+	RegisterTool(ToolManagerCommands.BeginMeshSelectionTool, TEXT("BeginMeshSelectionTool"), MeshSelectionToolBuilder);
+
+	auto EditMeshMaterialsToolBuilder = NewObject<UEditMeshMaterialsToolBuilder>();
+	RegisterTool(ToolManagerCommands.BeginEditMeshMaterialsTool, TEXT("BeginEditMeshMaterialsTool"), EditMeshMaterialsToolBuilder);
+	
+	RegisterTool(ToolManagerCommands.BeginMeshAttributePaintTool, TEXT("BeginMeshAttributePaintTool"), NewObject<UMeshAttributePaintToolBuilder>());
+
+	auto BakeMeshAttributeMapsToolBuilder = NewObject<UBakeMeshAttributeMapsToolBuilder>();
+	RegisterTool(ToolManagerCommands.BeginBakeMeshAttributeMapsTool, TEXT("BeginBakeMeshAttributeMapsTool"), BakeMeshAttributeMapsToolBuilder);
+
+	auto BakeMeshAttributeVertexToolBuilder = NewObject<UBakeMeshAttributeVertexToolBuilder>();
+	RegisterTool(ToolManagerCommands.BeginBakeMeshAttributeVertexTool, TEXT("BeginBakeMeshAttributeVertexTool"), BakeMeshAttributeVertexToolBuilder);
+
+	// analysis tools
+
+	RegisterTool(ToolManagerCommands.BeginMeshInspectorTool, TEXT("BeginMeshInspectorTool"), NewObject<UMeshInspectorToolBuilder>());
+	RegisterTool(ToolManagerCommands.BeginWeldEdgesTool, TEXT("BeginWeldEdgesTool"), NewObject<UWeldMeshEdgesToolBuilder>());
+	RegisterTool(ToolManagerCommands.BeginPolyGroupsTool, TEXT("BeginPolyGroupsTool"), NewObject<UConvertToPolygonsToolBuilder>());
+	RegisterTool(ToolManagerCommands.BeginAttributeEditorTool, TEXT("BeginAttributeEditorTool"), NewObject<UAttributeEditorToolBuilder>());
+
+
+	// Physics Tools
+
+	RegisterTool(ToolManagerCommands.BeginPhysicsInspectorTool, TEXT("BeginPhysicsInspectorTool"), NewObject<UPhysicsInspectorToolBuilder>());
+	RegisterTool(ToolManagerCommands.BeginSetCollisionGeometryTool, TEXT("BeginSetCollisionGeometryTool"), NewObject<USetCollisionGeometryToolBuilder>());
+	//RegisterTool(ToolManagerCommands.BeginEditCollisionGeometryTool, TEXT("EditCollisionGeoTool"), NewObject<UEditCollisionGeometryToolBuilder>());
+
+	auto ExtractCollisionGeoToolBuilder = NewObject<UExtractCollisionGeometryToolBuilder>();
+	RegisterTool(ToolManagerCommands.BeginExtractCollisionGeometryTool, TEXT("BeginExtractCollisionGeometryTool"), ExtractCollisionGeoToolBuilder);
+
+
+
+	// (experimental) hair tools
+
+	UGroomToMeshToolBuilder* GroomToMeshToolBuilder = NewObject<UGroomToMeshToolBuilder>();
+	RegisterTool(ToolManagerCommands.BeginGroomToMeshTool, TEXT("BeginGroomToMeshTool"), GroomToMeshToolBuilder);
+
+	RegisterTool(ToolManagerCommands.BeginGroomCardsEditorTool, TEXT("BeginGroomCardsEditorTool"), NewObject<UGroomCardsEditorToolBuilder>());
+
+	UGenerateLODMeshesToolBuilder* GenerateLODMeshesToolBuilder = NewObject<UGenerateLODMeshesToolBuilder>();
+	RegisterTool(ToolManagerCommands.BeginGenerateLODMeshesTool, TEXT("BeginGenerateLODMeshesTool"), GenerateLODMeshesToolBuilder);
+
+	// PolyModeling tools
+	auto RegisterPolyModelSelectTool = [&](EEditMeshPolygonsToolSelectionMode SelectionMode, TSharedPtr<FUICommandInfo> UICommand, FString StringName)
+	{
+		UEditMeshPolygonsSelectionModeToolBuilder* SelectionModeBuilder = NewObject<UEditMeshPolygonsSelectionModeToolBuilder>();
+		SelectionModeBuilder->SelectionMode = SelectionMode;
+		RegisterTool(UICommand, StringName, SelectionModeBuilder);
+	};
+	RegisterPolyModelSelectTool(EEditMeshPolygonsToolSelectionMode::Faces, ToolManagerCommands.BeginPolyModelTool_FaceSelect, TEXT("PolyEdit_FaceSelect"));
+	RegisterPolyModelSelectTool(EEditMeshPolygonsToolSelectionMode::Edges, ToolManagerCommands.BeginPolyModelTool_EdgeSelect, TEXT("PolyEdit_EdgeSelect"));
+	RegisterPolyModelSelectTool(EEditMeshPolygonsToolSelectionMode::Vertices, ToolManagerCommands.BeginPolyModelTool_VertexSelect, TEXT("PolyEdit_VertexSelect"));
+	RegisterPolyModelSelectTool(EEditMeshPolygonsToolSelectionMode::Loops, ToolManagerCommands.BeginPolyModelTool_LoopSelect, TEXT("PolyEdit_LoopSelect"));
+	RegisterPolyModelSelectTool(EEditMeshPolygonsToolSelectionMode::Rings, ToolManagerCommands.BeginPolyModelTool_RingSelect, TEXT("PolyEdit_RingSelect"));
+	RegisterPolyModelSelectTool(EEditMeshPolygonsToolSelectionMode::FacesEdgesVertices, ToolManagerCommands.BeginPolyModelTool_AllSelect, TEXT("PolyEdit_AllSelect"));
+
+	auto RegisterPolyModelActionTool = [&](EEditMeshPolygonsToolActions Action, TSharedPtr<FUICommandInfo> UICommand, FString StringName)
+	{
+		UEditMeshPolygonsActionModeToolBuilder* ActionModeBuilder = NewObject<UEditMeshPolygonsActionModeToolBuilder>();
+		ActionModeBuilder->StartupAction = Action;
+		RegisterTool(UICommand, StringName, ActionModeBuilder);
+	};
+	RegisterPolyModelActionTool(EEditMeshPolygonsToolActions::Extrude, ToolManagerCommands.BeginPolyModelTool_Extrude, TEXT("PolyEdit_Extrude"));
+	RegisterPolyModelActionTool(EEditMeshPolygonsToolActions::InsetOutset, ToolManagerCommands.BeginPolyModelTool_InsetOutset, TEXT("PolyEdit_InsetOutset"));
+	RegisterPolyModelActionTool(EEditMeshPolygonsToolActions::CutFaces, ToolManagerCommands.BeginPolyModelTool_CutFaces, TEXT("PolyEdit_CutFaces"));
+
+
+	// register extensions
+	TArray<IModelingModeToolExtension*> Extensions = IModularFeatures::Get().GetModularFeatureImplementations<IModelingModeToolExtension>(
+		IModelingModeToolExtension::GetModularFeatureName());
+	if (Extensions.Num() > 0)
+	{
+		FExtensionToolQueryInfo ExtensionQueryInfo;
+		ExtensionQueryInfo.ToolsContext = ToolsContext.Get();
+		ExtensionQueryInfo.AssetAPI = nullptr;
+
+		UE_LOG(LogTemp, Log, TEXT("ModelingMode: Found %d Tool Extension Modules"), Extensions.Num());
+		for (int32 k = 0; k < Extensions.Num(); ++k)
+		{
+			// TODO: extension name
+			FText ExtensionName = Extensions[k]->GetExtensionName();
+			FString ExtensionPrefix = FString::Printf(TEXT("[%d][%s]"), k, *ExtensionName.ToString());
+
+			TArray<FExtensionToolDescription> ToolSet;
+			Extensions[k]->GetExtensionTools(ExtensionQueryInfo, ToolSet);
+			for (const FExtensionToolDescription& ToolInfo : ToolSet)
+			{
+				UE_LOG(LogTemp, Log, TEXT("%s - Registering Tool [%s]"), *ExtensionPrefix, *ToolInfo.ToolName.ToString());
+
+				RegisterTool(ToolInfo.ToolCommand, ToolInfo.ToolName.ToString(), ToolInfo.ToolBuilder);
+			}
+		}
+	}
+
+
+	ToolsContext->ToolManager->SelectActiveToolType(EToolSide::Left, TEXT("DynaSculptTool"));
+
+	// Register modeling mode hotkeys. Note that we use the toolkit command list because we would like the hotkeys
+	// to work even when the viewport is not focused, provided that nothing else captures the key presses.
+	FModelingModeActionCommands::RegisterCommandBindings(Toolkit->GetToolkitCommands(), [this](EModelingModeActionCommands Command) {
+		ModelingModeShortcutRequested(Command);
+	});
+
+	// enable realtime viewport override
+	ConfigureRealTimeViewportsOverride(true);
+
+	//
+	// Engine Analytics
+	//
+
+	// Log mode starting
+	if (FEngineAnalytics::IsAvailable())
+	{
+		UModelingToolsEditorMode::LastModeStartTimestamp = FDateTime::UtcNow();
+		TArray<FAnalyticsEventAttribute> Attributes;
+		Attributes.Add(FAnalyticsEventAttribute(TEXT("Timestamp"), LastModeStartTimestamp.ToString()));
+		FEngineAnalytics::GetProvider().RecordEvent(TEXT("Editor.Usage.MeshModelingMode.Enter"), Attributes);
+	}
+
+	// Log tool starting
+	ToolsContext->ToolManager->OnToolStarted.AddLambda([this](UInteractiveToolManager* Manager, UInteractiveTool* Tool)
+	{
+		if (FEngineAnalytics::IsAvailable())
+		{
+			UModelingToolsEditorMode::LastToolStartTimestamp = FDateTime::UtcNow();
+			TArray<FAnalyticsEventAttribute> Attributes;
+			Attributes.Add(FAnalyticsEventAttribute(TEXT("DisplayName"), Tool->GetToolInfo().ToolDisplayName.ToString()));
+			Attributes.Add(FAnalyticsEventAttribute(TEXT("Timestamp"), LastToolStartTimestamp.ToString()));
+			FEngineAnalytics::GetProvider().RecordEvent(TEXT("Editor.Usage.MeshModelingMode.ToolStarted"), Attributes);
+		}
+	});
+
+	// Log tool ending
+	ToolsContext->ToolManager->OnToolEnded.AddLambda([this](UInteractiveToolManager* Manager, UInteractiveTool* Tool)
+	{
+		if (FEngineAnalytics::IsAvailable())
+		{
+			TArray<FAnalyticsEventAttribute> Attributes;
+			Attributes.Add(FAnalyticsEventAttribute(TEXT("DisplayName"), Tool->GetToolInfo().ToolDisplayName.ToString()));
+			Attributes.Add(FAnalyticsEventAttribute(TEXT("Timestamp"), FDateTime::UtcNow().ToString()));
+			const FTimespan ToolUsageDuration = FDateTime::UtcNow() - UModelingToolsEditorMode::LastToolStartTimestamp;
+			Attributes.Add(FAnalyticsEventAttribute(TEXT("Duration.Seconds"), static_cast<float>(ToolUsageDuration.GetTotalSeconds())));
+			FEngineAnalytics::GetProvider().RecordEvent(TEXT("Editor.Usage.MeshModelingMode.ToolEnded"), Attributes);
+		}
+	});
+
+	// do any toolkit UI initialization that depends on the mode setup above
+	if (Toolkit.IsValid())
+	{
+		FModelingToolsEditorModeToolkit* ModelingToolkit = (FModelingToolsEditorModeToolkit*)Toolkit.Get();
+		ModelingToolkit->InitializeAfterModeSetup();
+	}
+
+}
+
+void UModelingToolsEditorMode::Exit()
+{
+	//
+	// Engine Analytics
+	//
+	// Log mode ending
+	if (FEngineAnalytics::IsAvailable())
+	{
+		TArray<FAnalyticsEventAttribute> Attributes;
+		Attributes.Add(FAnalyticsEventAttribute(TEXT("Timestamp"), FDateTime::UtcNow().ToString()));
+		const FTimespan ModeUsageDuration = FDateTime::UtcNow() - UModelingToolsEditorMode::LastModeStartTimestamp;
+		Attributes.Add(FAnalyticsEventAttribute(TEXT("Duration.Seconds"), static_cast<float>(ModeUsageDuration.GetTotalSeconds())));
+		FEngineAnalytics::GetProvider().RecordEvent(TEXT("Editor.Usage.MeshModelingMode.Exit"), Attributes);
+	}
+
+	StylusStateTracker = nullptr;
+
+	// TODO: cannot deregister currently because if another mode is also registering, its Enter()
+	// will be called before our Exit()
+	//UE::TransformGizmoUtil::DeregisterTransformGizmoContextObject(ToolsContext.Get());
+	
+	// TODO: cannot deregister currently because if another mode is also registering, its Enter()
+	// will be called before our Exit()
+	UEditorModelingObjectsCreationAPI* ObjectCreationAPI = UEditorModelingObjectsCreationAPI::Find(ToolsContext.Get());
+	if (ObjectCreationAPI)
+	{
+		ObjectCreationAPI->GetNewAssetPathNameCallback.Unbind();
+		ObjectCreationAPI->OnModelingMeshCreated.Remove(MeshCreatedEventHandle);
+		ObjectCreationAPI->OnModelingTextureCreated.Remove(TextureCreatedEventHandle);
+		//UEditorModelingObjectsCreationAPI::Deregister(ToolsContext.Get());		// cannot do currently because of shared ToolsContext, revisit in future
+	}
+
+	FModelingModeActionCommands::UnRegisterCommandBindings(Toolkit->GetToolkitCommands());
+
+	// clear realtime viewport override
+	ConfigureRealTimeViewportsOverride(false);
+
+	// re-enable HitProxy rendering
+	ToolsContext.Get()->SetEnableRenderingDuringHitProxyPass(true);
+
+	// Call base Exit method to ensure proper cleanup
+	UEdMode::Exit();
+}
+
+void UModelingToolsEditorMode::CreateToolkit()
+{
+	Toolkit = MakeShareable(new FModelingToolsEditorModeToolkit);
+}
+
+void UModelingToolsEditorMode::OnToolStarted(UInteractiveToolManager* Manager, UInteractiveTool* Tool)
+{
+	FModelingToolActionCommands::UpdateToolCommandBinding(Tool, Toolkit->GetToolkitCommands(), false);
+	
+	if( FEngineAnalytics::IsAvailable() )
+	{
+		FEngineAnalytics::GetProvider().RecordEvent( TEXT( "Editor.Usage.MeshModelingMode.ToolStarted" ),
+													 TEXT( "DisplayName" ),
+													 Tool->GetToolInfo().ToolDisplayName.ToString() );
+	}
+}
+
+void UModelingToolsEditorMode::OnToolEnded(UInteractiveToolManager* Manager, UInteractiveTool* Tool)
+{
+	FModelingToolActionCommands::UpdateToolCommandBinding(Tool, Toolkit->GetToolkitCommands(), true);
+	
+	if( FEngineAnalytics::IsAvailable() )
+	{
+		FEngineAnalytics::GetProvider().RecordEvent( TEXT( "Editor.Usage.MeshModelingMode.ToolEnded" ),
+													 TEXT( "DisplayName" ),
+													 Tool->GetToolInfo().ToolDisplayName.ToString() );
+	}
+}
+
+void UModelingToolsEditorMode::BindCommands()
+{
+	const FModelingToolsManagerCommands& ToolManagerCommands = FModelingToolsManagerCommands::Get();
+	const TSharedRef<FUICommandList>& CommandList = Toolkit->GetToolkitCommands();
+
+	CommandList->MapAction(
+		ToolManagerCommands.AcceptActiveTool,
+		FExecuteAction::CreateLambda([this]() { ToolsContext->EndTool(EToolShutdownType::Accept); }),
+		FCanExecuteAction::CreateLambda([this]() { return ToolsContext->CanAcceptActiveTool(); }),
+		FGetActionCheckState(),
+		FIsActionButtonVisible::CreateLambda([this]() {return ToolsContext->ActiveToolHasAccept(); }),
+		EUIActionRepeatMode::RepeatDisabled
+		);
+
+	CommandList->MapAction(
+		ToolManagerCommands.CancelActiveTool,
+		FExecuteAction::CreateLambda([this]() { ToolsContext->EndTool(EToolShutdownType::Cancel); }),
+		FCanExecuteAction::CreateLambda([this]() { return ToolsContext->CanCancelActiveTool(); }),
+		FGetActionCheckState(),
+		FIsActionButtonVisible::CreateLambda([this]() {return ToolsContext->ActiveToolHasAccept(); }),
+		EUIActionRepeatMode::RepeatDisabled
+		);
+
+	CommandList->MapAction(
+		ToolManagerCommands.CompleteActiveTool,
+		FExecuteAction::CreateLambda([this]() { ToolsContext->EndTool(EToolShutdownType::Completed); }),
+		FCanExecuteAction::CreateLambda([this]() { return ToolsContext->CanCompleteActiveTool(); }),
+		FGetActionCheckState(),
+		FIsActionButtonVisible::CreateLambda([this]() {return ToolsContext->CanCompleteActiveTool(); }),
+		EUIActionRepeatMode::RepeatDisabled
+		);
+
+	// These aren't activated by buttons but have default chords that bind the keypresses to the action.
+	CommandList->MapAction(
+		ToolManagerCommands.AcceptOrCompleteActiveTool,
+		FExecuteAction::CreateLambda([this]() { AcceptActiveToolActionOrTool(); }),
+		FCanExecuteAction::CreateLambda([this]() {
+				return ToolsContext->CanAcceptActiveTool() || ToolsContext->CanCompleteActiveTool();
+			}),
+		FGetActionCheckState(),
+		FIsActionButtonVisible(),
+		EUIActionRepeatMode::RepeatDisabled);
+
+	CommandList->MapAction(
+		ToolManagerCommands.CancelOrCompleteActiveTool,
+		FExecuteAction::CreateLambda([this]() { CancelActiveToolActionOrTool(); }),
+		FCanExecuteAction::CreateLambda([this]() {
+				return ToolsContext->CanCompleteActiveTool() || ToolsContext->CanCancelActiveTool();
+			}),
+		FGetActionCheckState(),
+		FIsActionButtonVisible(),
+		EUIActionRepeatMode::RepeatDisabled);
+}
+
+
+void UModelingToolsEditorMode::AcceptActiveToolActionOrTool()
+{
+	// if we have an active Tool that implements 
+	if (GetToolManager()->HasAnyActiveTool())
+	{
+		UInteractiveTool* Tool = GetToolManager()->GetActiveTool(EToolSide::Mouse);
+		IInteractiveToolNestedAcceptCancelAPI* CancelAPI = Cast<IInteractiveToolNestedAcceptCancelAPI>(Tool);
+		if (CancelAPI && CancelAPI->SupportsNestedAcceptCommand() && CancelAPI->CanCurrentlyNestedAccept())
+		{
+			bool bAccepted = CancelAPI->ExecuteNestedAcceptCommand();
+			if (bAccepted)
+			{
+				return;
+			}
+		}
+	}
+
+	const EToolShutdownType ShutdownType = ToolsContext->CanAcceptActiveTool() ? EToolShutdownType::Accept : EToolShutdownType::Completed;
+	ToolsContext->EndTool(ShutdownType);
+}
+
+
+void UModelingToolsEditorMode::CancelActiveToolActionOrTool()
+{
+	// if we have an active Tool that implements 
+	if (GetToolManager()->HasAnyActiveTool())
+	{
+		UInteractiveTool* Tool = GetToolManager()->GetActiveTool(EToolSide::Mouse);
+		IInteractiveToolNestedAcceptCancelAPI* CancelAPI = Cast<IInteractiveToolNestedAcceptCancelAPI>(Tool);
+		if (CancelAPI && CancelAPI->SupportsNestedCancelCommand() && CancelAPI->CanCurrentlyNestedCancel())
+		{
+			bool bCancelled = CancelAPI->ExecuteNestedCancelCommand();
+			if (bCancelled)
+			{
+				return;
+			}
+		}
+	}
+
+	const EToolShutdownType ShutdownType = ToolsContext->CanCancelActiveTool() ? EToolShutdownType::Cancel : EToolShutdownType::Completed;
+	ToolsContext->EndTool(ShutdownType);
+}
+
+void UModelingToolsEditorMode::ModelingModeShortcutRequested(EModelingModeActionCommands Command)
+{
+	if (Command == EModelingModeActionCommands::FocusViewToCursor)
+	{
+		FocusCameraAtCursorHotkey();
+	}
+}
+
+
+void UModelingToolsEditorMode::FocusCameraAtCursorHotkey()
+{
+	FRay Ray = ToolsContext->GetLastWorldRay();
+
+	double NearestHitDist = (double)HALF_WORLD_MAX;
+	FVector HitPoint = FVector::ZeroVector;
+
+	// cast ray against visible objects
+	FHitResult WorldHitResult;
+	if (ToolSceneQueriesUtil::FindNearestVisibleObjectHit(GetWorld(), WorldHitResult, Ray.Origin, Ray.PointAt(HALF_WORLD_MAX)))
+	{
+		HitPoint = WorldHitResult.ImpactPoint;
+		NearestHitDist = (double)Ray.GetParameter(HitPoint);
+	}
+
+	// cast ray against tool
+	if (GetToolManager()->HasAnyActiveTool())
+	{
+		UInteractiveTool* Tool = GetToolManager()->GetActiveTool(EToolSide::Mouse);
+		IInteractiveToolCameraFocusAPI* FocusAPI = Cast<IInteractiveToolCameraFocusAPI>(Tool);
+		if (FocusAPI && FocusAPI->SupportsWorldSpaceFocusPoint())
+		{
+			FVector ToolHitPoint;
+			if (FocusAPI->GetWorldSpaceFocusPoint(Ray, ToolHitPoint))
+			{
+				double HitDepth = (double)Ray.GetParameter(ToolHitPoint);
+				if (HitDepth < NearestHitDist)
+				{
+					NearestHitDist = HitDepth;
+					HitPoint = ToolHitPoint;
+				}
+			}
+		}
+	}
+
+
+	if (NearestHitDist < (double)HALF_WORLD_MAX && GCurrentLevelEditingViewportClient)
+	{
+		GCurrentLevelEditingViewportClient->CenterViewportAtPoint(HitPoint, false);
+	}
+}
+
+
+bool UModelingToolsEditorMode::ComputeBoundingBoxForViewportFocus(AActor* Actor, UPrimitiveComponent* PrimitiveComponent, FBox& InOutBox) const
+{
+	if (GetToolManager()->HasAnyActiveTool())
+	{
+		UInteractiveTool* Tool = GetToolManager()->GetActiveTool(EToolSide::Mouse);
+		IInteractiveToolCameraFocusAPI* FocusAPI = Cast<IInteractiveToolCameraFocusAPI>(Tool);
+		if (FocusAPI && FocusAPI->SupportsWorldSpaceFocusBox() )
+		{
+			InOutBox = FocusAPI->GetWorldSpaceFocusBox();
+			if (InOutBox.IsValid)
+			{
+				float MaxDimension = InOutBox.GetExtent().GetMax();
+				if (MaxDimension > SMALL_NUMBER)
+				{
+					InOutBox.ExpandBy(MaxDimension * 0.2f);
+				}
+				else
+				{
+					InOutBox.ExpandBy(25);
+				}
+				return true;
+			}
+		}
+	}
+	return false;
+}
+
+
+bool UModelingToolsEditorMode::GetPivotForOrbit(FVector& OutPivot) const
+{
+	if (GCurrentLevelEditingViewportClient)
+	{
+		OutPivot = GCurrentLevelEditingViewportClient->GetViewTransform().GetLookAt();
+		return true;
+	}
+	return false;
+}
+
+
+
+void UModelingToolsEditorMode::ConfigureRealTimeViewportsOverride(bool bEnable)
+{
+	FLevelEditorModule& LevelEditorModule = FModuleManager::GetModuleChecked<FLevelEditorModule>("LevelEditor");
+	TSharedPtr<ILevelEditor> LevelEditor = LevelEditorModule.GetFirstLevelEditor();
+	if (LevelEditor.IsValid())
+	{
+		TArray<TSharedPtr<SLevelViewport>> Viewports = LevelEditor->GetViewports();
+		for (const TSharedPtr<SLevelViewport>& ViewportWindow : Viewports)
+		{
+			if (ViewportWindow.IsValid())
+			{
+				FEditorViewportClient& Viewport = ViewportWindow->GetAssetViewportClient();
+				const FText SystemDisplayName = LOCTEXT("RealtimeOverrideMessage_ModelingMode", "Modeling Mode");
+				if (bEnable)
+				{
+					Viewport.AddRealtimeOverride(bEnable, SystemDisplayName);
+				}
+				else
+				{
+					Viewport.RemoveRealtimeOverride(SystemDisplayName, false);
+				}
+			}
+		}
+	}
+}
+
+
+
+#undef LOCTEXT_NAMESPACE
diff --git a/Engine/Plugins/Editor/ModelingToolsEditorMode/Source/ModelingToolsEditorMode/Private/ModelingToolsEditorModeModule.cpp b/Engine/Plugins/Editor/ModelingToolsEditorMode/Source/ModelingToolsEditorMode/Private/ModelingToolsEditorModeModule.cpp
new file mode 100644
--- /dev/null
+++ b/Engine/Plugins/Editor/ModelingToolsEditorMode/Source/ModelingToolsEditorMode/Private/ModelingToolsEditorModeModule.cpp
@@ -0,0 +1,70 @@
+// Copyright Epic Games, Inc. All Rights Reserved.
+
+#include "ModelingToolsEditorModeModule.h"
+
+#include "ModelingToolsActions.h"
+#include "ModelingToolsManagerActions.h"
+#include "ModelingToolsEditorModeStyle.h"
+#include "ModelingToolsEditorModeSettings.h"
+
+#include "ISettingsModule.h"
+#include "ISettingsSection.h"
+#include "Misc/CoreDelegates.h"
+
+#include "Modules/ModuleManager.h"
+#include "PropertyEditorModule.h"
+#include "DetailsCustomizations/ModelingToolPropertyCustomizations.h"
+
+#define LOCTEXT_NAMESPACE "FModelingToolsEditorModeModule"
+
+void FModelingToolsEditorModeModule::StartupModule()
+{
+	FCoreDelegates::OnPostEngineInit.AddRaw(this, &FModelingToolsEditorModeModule::OnPostEngineInit);
+}
+
+void FModelingToolsEditorModeModule::ShutdownModule()
+{
+	FCoreDelegates::OnPostEngineInit.RemoveAll(this);
+
+	if (ISettingsModule* SettingsModule = FModuleManager::GetModulePtr<ISettingsModule>("Settings"))
+	{
+		SettingsModule->UnregisterSettings("Project", "Plugins", "ModelingMode");
+	}
+
+	FModelingToolActionCommands::UnregisterAllToolActions();
+	FModelingToolsManagerCommands::Unregister();
+	FModelingModeActionCommands::Unregister();
+
+	// Unregister slate style overrides
+	FModelingToolsEditorModeStyle::Shutdown();
+}
+
+
+void FModelingToolsEditorModeModule::OnPostEngineInit()
+{
+	// Register slate style overrides
+	FModelingToolsEditorModeStyle::Initialize();
+
+	FModelingToolActionCommands::RegisterAllToolActions();
+	FModelingToolsManagerCommands::Register();
+	FModelingModeActionCommands::Register();
+
+	// register settings
+	ISettingsModule* SettingsModule = FModuleManager::GetModulePtr<ISettingsModule>("Settings");
+	if (SettingsModule != nullptr)
+	{
+		ISettingsSectionPtr SettingsSection = SettingsModule->RegisterSettings("Project", "Plugins", "ModelingMode",
+			LOCTEXT("ModelingModeSettingsName", "Modeling Mode"),
+			LOCTEXT("ModelingModeSettingsDescription", "Configure the Modeling Tools Editor Mode plugin"),
+			GetMutableDefault<UModelingToolsEditorModeSettings>()
+		);
+	}
+
+	// Register details view customizations
+	FPropertyEditorModule& PropertyModule = FModuleManager::LoadModuleChecked<FPropertyEditorModule>("PropertyEditor");
+	PropertyModule.RegisterCustomPropertyTypeLayout("ModelingToolsAxisFilter", FOnGetPropertyTypeCustomizationInstance::CreateStatic(&FModelingToolsAxisFilterCustomization::MakeInstance));
+}
+
+#undef LOCTEXT_NAMESPACE
+	
+IMPLEMENT_MODULE(FModelingToolsEditorModeModule, ModelingToolsEditorMode)
\ No newline at end of file
diff --git a/Engine/Plugins/Editor/ModelingToolsEditorMode/Source/ModelingToolsEditorMode/Private/ModelingToolsEditorModeStyle.cpp b/Engine/Plugins/Editor/ModelingToolsEditorMode/Source/ModelingToolsEditorMode/Private/ModelingToolsEditorModeStyle.cpp
new file mode 100644
--- /dev/null
+++ b/Engine/Plugins/Editor/ModelingToolsEditorMode/Source/ModelingToolsEditorMode/Private/ModelingToolsEditorModeStyle.cpp
@@ -0,0 +1,313 @@
+// Copyright Epic Games, Inc. All Rights Reserved.
+
+#include "ModelingToolsEditorModeStyle.h"
+#include "Styling/SlateStyleRegistry.h"
+#include "Styling/SlateTypes.h"
+#include "Styling/CoreStyle.h"
+#include "EditorStyleSet.h"
+#include "Interfaces/IPluginManager.h"
+#include "SlateOptMacros.h"
+#include "Styling/SlateStyleMacros.h"
+
+
+#define IMAGE_PLUGIN_BRUSH( RelativePath, ... ) FSlateImageBrush( FModelingToolsEditorModeStyle::InContent( RelativePath, ".png" ), __VA_ARGS__ )
+
+// This is to fix the issue that SlateStyleMacros like IMAGE_BRUSH look for RootToContentDir but StyleSet->RootToContentDir is how this style is set up
+#define RootToContentDir StyleSet->RootToContentDir
+
+FString FModelingToolsEditorModeStyle::InContent(const FString& RelativePath, const ANSICHAR* Extension)
+{
+	static FString ContentDir = IPluginManager::Get().FindPlugin(TEXT("ModelingToolsEditorMode"))->GetContentDir();
+	return (ContentDir / RelativePath) + Extension;
+}
+
+TSharedPtr< FSlateStyleSet > FModelingToolsEditorModeStyle::StyleSet = nullptr;
+TSharedPtr< class ISlateStyle > FModelingToolsEditorModeStyle::Get() { return StyleSet; }
+
+FName FModelingToolsEditorModeStyle::GetStyleSetName()
+{
+	static FName ModelingToolsStyleName(TEXT("ModelingToolsStyle"));
+	return ModelingToolsStyleName;
+}
+
+BEGIN_SLATE_FUNCTION_BUILD_OPTIMIZATION
+
+void FModelingToolsEditorModeStyle::Initialize()
+{
+	// Const icon sizes
+	const FVector2D Icon8x8(8.0f, 8.0f);
+	const FVector2D Icon16x16(16.0f, 16.0f);
+	const FVector2D Icon20x20(20.0f, 20.0f);
+	const FVector2D Icon28x28(28.0f, 28.0f);
+	const FVector2D Icon40x40(40.0f, 40.0f);
+
+	// Only register once
+	if (StyleSet.IsValid())
+	{
+		return;
+	}
+
+	StyleSet = MakeShareable(new FSlateStyleSet(GetStyleSetName()));
+	StyleSet->SetContentRoot(FPaths::EnginePluginsDir() / TEXT("Editor/ModelingToolsEditorMode/Content"));
+	StyleSet->SetCoreContentRoot(FPaths::EngineContentDir() / TEXT("Slate"));
+
+	const FTextBlockStyle& NormalText = FEditorStyle::Get().GetWidgetStyle<FTextBlockStyle>("NormalText");
+
+	// Shared editors
+	//{
+	//	StyleSet->Set("Paper2D.Common.ViewportZoomTextStyle", FTextBlockStyle(NormalText)
+	//		.SetFont(DEFAULT_FONT("BoldCondensed", 16))
+	//	);
+
+	//	StyleSet->Set("Paper2D.Common.ViewportTitleTextStyle", FTextBlockStyle(NormalText)
+	//		.SetFont(DEFAULT_FONT("Regular", 18))
+	//		.SetColorAndOpacity(FLinearColor(1.0, 1.0f, 1.0f, 0.5f))
+	//	);
+
+	//	StyleSet->Set("Paper2D.Common.ViewportTitleBackground", new BOX_BRUSH("Old/Graph/GraphTitleBackground", FMargin(0)));
+	//}
+
+	// Tool Manager icons
+	{
+		// Accept/Cancel/Complete active tool
+
+		if (FCoreStyle::IsStarshipStyle())
+		{
+			StyleSet->Set("LevelEditor.ModelingToolsMode", new IMAGE_BRUSH_SVG("Starship/geometry", FVector2D(20.0f, 20.0f)));
+		}
+		else
+		{
+			StyleSet->Set("LevelEditor.ModelingToolsMode", new IMAGE_PLUGIN_BRUSH("Icons/icon_ModelingToolsEditorMode", FVector2D(40.0f, 40.0f)));
+			StyleSet->Set("LevelEditor.ModelingToolsMode.Small", new IMAGE_PLUGIN_BRUSH("Icons/icon_ModelingToolsEditorMode", FVector2D(20.0f, 20.0f)));
+		}
+
+		// NOTE:  Old-style, need to be replaced: 
+		StyleSet->Set("ModelingToolsManagerCommands.CancelActiveTool", new IMAGE_PLUGIN_BRUSH("Icons/icon_ActiveTool_Cancel_40x", Icon20x20));
+		StyleSet->Set("ModelingToolsManagerCommands.CancelActiveTool.Small", new IMAGE_PLUGIN_BRUSH("Icons/icon_ActiveTool_Cancel_40x", Icon20x20));
+		StyleSet->Set("ModelingToolsManagerCommands.AcceptActiveTool", new IMAGE_PLUGIN_BRUSH("Icons/icon_ActiveTool_Accept_40x", Icon20x20));
+		StyleSet->Set("ModelingToolsManagerCommands.AcceptActiveTool.Small", new IMAGE_PLUGIN_BRUSH("Icons/icon_ActiveTool_Accept_40x", Icon20x20));
+		StyleSet->Set("ModelingToolsManagerCommands.CompleteActiveTool", new IMAGE_PLUGIN_BRUSH("Icons/icon_ActiveTool_Accept_40x", Icon20x20));
+		StyleSet->Set("ModelingToolsManagerCommands.CompleteActiveTool.Small", new IMAGE_PLUGIN_BRUSH("Icons/icon_ActiveTool_Accept_40x", Icon20x20));
+
+
+		StyleSet->Set("ModelingToolsManagerCommands.BeginShapeSprayTool", 				new IMAGE_PLUGIN_BRUSH("Icons/ShapeSpray_40x",	Icon20x20));
+		StyleSet->Set("ModelingToolsManagerCommands.BeginShapeSprayTool.Small", 		new IMAGE_PLUGIN_BRUSH("Icons/ShapeSpray_40x",	Icon20x20));
+		//StyleSet->Set("ModelingToolsManagerCommands.BeginMeshSpaceDeformerTool", 		new IMAGE_PLUGIN_BRUSH("Icons/icon_Tool_Displace_40x",		Icon20x20));
+		//StyleSet->Set("ModelingToolsManagerCommands.BeginMeshSpaceDeformerTool.Small", 	new IMAGE_PLUGIN_BRUSH("Icons/icon_Tool_Displace_40x",		Icon20x20));
+		StyleSet->Set("ModelingToolsManagerCommands.BeginPolygonOnMeshTool", 			new IMAGE_PLUGIN_BRUSH("Icons/icon_Tool_PolygonOnMesh_40x",	Icon20x20));
+		StyleSet->Set("ModelingToolsManagerCommands.BeginPolygonOnMeshTool.Small", 		new IMAGE_PLUGIN_BRUSH("Icons/icon_Tool_PolygonOnMesh_40x",	Icon20x20));
+		StyleSet->Set("ModelingToolsManagerCommands.BeginParameterizeMeshTool", 		new IMAGE_PLUGIN_BRUSH("Icons/icon_Tool_UVGenerate_40x",	Icon20x20));
+		StyleSet->Set("ModelingToolsManagerCommands.BeginParameterizeMeshTool.Small", 	new IMAGE_PLUGIN_BRUSH("Icons/icon_Tool_UVGenerate_40x",	Icon20x20));
+		StyleSet->Set("ModelingToolsManagerCommands.BeginPolyGroupsTool", 				new IMAGE_PLUGIN_BRUSH("Icons/icon_Tool_PolyGroups_40x",	Icon20x20));
+		StyleSet->Set("ModelingToolsManagerCommands.BeginPolyGroupsTool.Small", 		new IMAGE_PLUGIN_BRUSH("Icons/icon_Tool_PolyGroups_40x",	Icon20x20));
+
+
+		// Modes Palette Toolbar Icons
+		StyleSet->Set("ModelingToolsManagerCommands.BeginAddBoxPrimitiveTool", 			new IMAGE_BRUSH_SVG("Icons/ModelingBox", 		Icon20x20));
+		StyleSet->Set("ModelingToolsManagerCommands.BeginAddBoxPrimitiveTool.Small", 		new IMAGE_BRUSH_SVG("Icons/ModelingBox",		Icon40x40));
+		StyleSet->Set("ModelingToolsManagerCommands.BeginAddCylinderPrimitiveTool", 			new IMAGE_BRUSH_SVG("Icons/ModelingCylinder", 		Icon20x20));
+		StyleSet->Set("ModelingToolsManagerCommands.BeginAddCylinderPrimitiveTool.Small", 		new IMAGE_BRUSH_SVG("Icons/ModelingCylinder",		Icon40x40));
+		StyleSet->Set("ModelingToolsManagerCommands.BeginAddConePrimitiveTool", 			new IMAGE_BRUSH_SVG("Icons/ModelingCone", 		Icon20x20));
+		StyleSet->Set("ModelingToolsManagerCommands.BeginAddConePrimitiveTool.Small", 		new IMAGE_BRUSH_SVG("Icons/ModelingCone",		Icon40x40));
+		StyleSet->Set("ModelingToolsManagerCommands.BeginAddArrowPrimitiveTool", 			new IMAGE_BRUSH_SVG("Icons/ModelingArrow", 		Icon20x20));
+		StyleSet->Set("ModelingToolsManagerCommands.BeginAddArrowPrimitiveTool.Small", 		new IMAGE_BRUSH_SVG("Icons/ModelingArrow",		Icon40x40));
+		StyleSet->Set("ModelingToolsManagerCommands.BeginAddRectanglePrimitiveTool", 			new IMAGE_PLUGIN_BRUSH("Icons/ModelingRectangle_x20", 		Icon20x20));
+		StyleSet->Set("ModelingToolsManagerCommands.BeginAddRectanglePrimitiveTool.Small", 		new IMAGE_PLUGIN_BRUSH("Icons/ModelingRectangle_x40",		Icon40x40));
+		StyleSet->Set("ModelingToolsManagerCommands.BeginAddDiscPrimitiveTool", 			new IMAGE_PLUGIN_BRUSH("Icons/ModelingDisc_x20", 		Icon20x20));
+		StyleSet->Set("ModelingToolsManagerCommands.BeginAddDiscPrimitiveTool.Small", 		new IMAGE_PLUGIN_BRUSH("Icons/ModelingDisc_x40",		Icon40x40));
+		StyleSet->Set("ModelingToolsManagerCommands.BeginAddTorusPrimitiveTool", 			new IMAGE_BRUSH_SVG("Icons/ModelingTorus", 		Icon20x20));
+		StyleSet->Set("ModelingToolsManagerCommands.BeginAddTorusPrimitiveTool.Small", 		new IMAGE_BRUSH_SVG("Icons/ModelingTorus",		Icon40x40));
+		StyleSet->Set("ModelingToolsManagerCommands.BeginAddSpherePrimitiveTool", 			new IMAGE_BRUSH_SVG("Icons/ModelingSphere", 		Icon20x20));
+		StyleSet->Set("ModelingToolsManagerCommands.BeginAddSpherePrimitiveTool.Small", 		new IMAGE_BRUSH_SVG("Icons/ModelingSphere",		Icon40x40));
+		StyleSet->Set("ModelingToolsManagerCommands.BeginAddStairsPrimitiveTool", 			new IMAGE_BRUSH_SVG("Icons/Staircase", Icon20x20));
+		StyleSet->Set("ModelingToolsManagerCommands.BeginAddStairsPrimitiveTool.Small",		new IMAGE_BRUSH_SVG("Icons/Staircase", Icon40x40));
+
+		StyleSet->Set("ModelingToolsManagerCommands.BeginDrawPolygonTool", 				new IMAGE_PLUGIN_BRUSH("Icons/DrawPolygon_40x",		Icon20x20));
+		StyleSet->Set("ModelingToolsManagerCommands.BeginDrawPolygonTool.Small", 		new IMAGE_PLUGIN_BRUSH("Icons/DrawPolygon_40x", 	Icon20x20));
+		StyleSet->Set("ModelingToolsManagerCommands.BeginAddPatchTool",					new IMAGE_PLUGIN_BRUSH("Icons/Patch_40x",			Icon20x20));
+		StyleSet->Set("ModelingToolsManagerCommands.BeginAddPatchTool.Small",			new IMAGE_PLUGIN_BRUSH("Icons/Patch_40x",			Icon20x20));
+
+		StyleSet->Set("ModelingToolsManagerCommands.BeginSmoothMeshTool", 				new IMAGE_PLUGIN_BRUSH("Icons/ModelingSmooth_x40", 			Icon20x20));
+		StyleSet->Set("ModelingToolsManagerCommands.BeginSmoothMeshTool.Small", 		new IMAGE_PLUGIN_BRUSH("Icons/ModelingSmooth_x40", 			Icon20x20));
+		StyleSet->Set("ModelingToolsManagerCommands.BeginSculptMeshTool", 				new IMAGE_PLUGIN_BRUSH("Icons/Sculpt_40x", 			Icon20x20));
+		StyleSet->Set("ModelingToolsManagerCommands.BeginSculptMeshTool.Small", 		new IMAGE_PLUGIN_BRUSH("Icons/Sculpt_40x", 			Icon20x20));
+		StyleSet->Set("ModelingToolsManagerCommands.BeginPolyEditTool", 				new IMAGE_PLUGIN_BRUSH("Icons/PolyEdit_40x", 		Icon20x20));
+		StyleSet->Set("ModelingToolsManagerCommands.BeginPolyEditTool.Small", 			new IMAGE_PLUGIN_BRUSH("Icons/PolyEdit_40x", 		Icon20x20));
+		StyleSet->Set("ModelingToolsManagerCommands.BeginSubdividePolyTool",			new IMAGE_BRUSH_SVG("Icons/ModelingSubD",		Icon20x20));
+		StyleSet->Set("ModelingToolsManagerCommands.BeginSubdividePolyTool.Small",		new IMAGE_BRUSH_SVG("Icons/ModelingSubD",		Icon20x20));
+		StyleSet->Set("ModelingToolsManagerCommands.BeginGroupEdgeInsertionTool", 		new IMAGE_PLUGIN_BRUSH("Icons/ModelingGroupEdgeInsert_x40", 		Icon20x20));
+		StyleSet->Set("ModelingToolsManagerCommands.BeginGroupEdgeInsertionTool.Small", new IMAGE_PLUGIN_BRUSH("Icons/ModelingGroupEdgeInsert_x40", 		Icon20x20));
+		StyleSet->Set("ModelingToolsManagerCommands.BeginEdgeLoopInsertionTool", 		new IMAGE_PLUGIN_BRUSH("Icons/ModelingEdgeLoopInsert_x40", 		Icon20x20));
+		StyleSet->Set("ModelingToolsManagerCommands.BeginEdgeLoopInsertionTool.Small",  new IMAGE_PLUGIN_BRUSH("Icons/ModelingEdgeLoopInsert_x40", 		Icon20x20));
+		StyleSet->Set("ModelingToolsManagerCommands.BeginTriEditTool", 				new IMAGE_PLUGIN_BRUSH("Icons/TriEdit_40x", 		Icon20x20));
+		StyleSet->Set("ModelingToolsManagerCommands.BeginTriEditTool.Small", 			new IMAGE_PLUGIN_BRUSH("Icons/TriEdit_40x", 		Icon20x20));
+		//StyleSet->Set("ModelingToolsManagerCommands.BeginPolyDeformTool", 			new IMAGE_PLUGIN_BRUSH("Icons/PolyEdit_40x", 		Icon20x20));
+		//StyleSet->Set("ModelingToolsManagerCommands.BeginPolyDeformTool.Small", 		new IMAGE_PLUGIN_BRUSH("Icons/PolyEdit_40x", 		Icon20x20));
+		StyleSet->Set("ModelingToolsManagerCommands.BeginDisplaceMeshTool", 			new IMAGE_PLUGIN_BRUSH("Icons/Displace_40x", 		Icon20x20));
+		StyleSet->Set("ModelingToolsManagerCommands.BeginDisplaceMeshTool.Small", 		new IMAGE_PLUGIN_BRUSH("Icons/Displace_40x", 		Icon20x20));
+		StyleSet->Set("ModelingToolsManagerCommands.BeginTransformMeshesTool", 			new IMAGE_PLUGIN_BRUSH("Icons/Transform_40x", 		Icon20x20));
+		StyleSet->Set("ModelingToolsManagerCommands.BeginTransformMeshesTool.Small", 	new IMAGE_PLUGIN_BRUSH("Icons/Transform_40x", 		Icon20x20));
+
+		StyleSet->Set("ModelingToolsManagerCommands.BeginRemeshSculptMeshTool", 		new IMAGE_PLUGIN_BRUSH("Icons/DynaSculpt_40x", 		Icon20x20));
+		StyleSet->Set("ModelingToolsManagerCommands.BeginRemeshSculptMeshTool.Small", 	new IMAGE_PLUGIN_BRUSH("Icons/DynaSculpt_40x",		Icon20x20));
+		StyleSet->Set("ModelingToolsManagerCommands.BeginRemeshMeshTool", 				new IMAGE_PLUGIN_BRUSH("Icons/Remesh_40x", 			Icon20x20));
+		StyleSet->Set("ModelingToolsManagerCommands.BeginRemeshMeshTool.Small", 		new IMAGE_PLUGIN_BRUSH("Icons/Remesh_40x", 			Icon20x20));
+		StyleSet->Set("ModelingToolsManagerCommands.BeginProjectToTargetTool", 		new IMAGE_PLUGIN_BRUSH("Icons/ModelingRemeshToTarget_x40",	Icon20x20));
+		StyleSet->Set("ModelingToolsManagerCommands.BeginProjectToTargetTool.Small", 	new IMAGE_PLUGIN_BRUSH("Icons/ModelingRemeshToTarget_x40",	Icon20x20));
+		//StyleSet->Set("ModelingToolsManagerCommands.BeginProjectToTargetTool",			new IMAGE_PLUGIN_BRUSH("",			Icon20x20));
+		//StyleSet->Set("ModelingToolsManagerCommands.BeginProjectToTargetTool.Small",	new IMAGE_PLUGIN_BRUSH("",			Icon20x20));
+		StyleSet->Set("ModelingToolsManagerCommands.BeginSimplifyMeshTool", 			new IMAGE_PLUGIN_BRUSH("Icons/Simplify_40x", 		Icon20x20));
+		StyleSet->Set("ModelingToolsManagerCommands.BeginSimplifyMeshTool.Small", 		new IMAGE_PLUGIN_BRUSH("Icons/Simplify_40x", 		Icon20x20));
+		StyleSet->Set("ModelingToolsManagerCommands.BeginEditNormalsTool", 			new IMAGE_PLUGIN_BRUSH("Icons/Normals_40x",			Icon20x20));
+		StyleSet->Set("ModelingToolsManagerCommands.BeginEditNormalsTool.Small", 		new IMAGE_PLUGIN_BRUSH("Icons/Normals_40x",			Icon20x20));
+		StyleSet->Set("ModelingToolsManagerCommands.BeginEditTangentsTool", 			new IMAGE_PLUGIN_BRUSH("Icons/ModelingTangents_x40",			Icon20x20));
+		StyleSet->Set("ModelingToolsManagerCommands.BeginEditTangentsTool.Small", 		new IMAGE_PLUGIN_BRUSH("Icons/ModelingTangents_x40",			Icon20x20));
+		StyleSet->Set("ModelingToolsManagerCommands.BeginUVSeamEditTool", 			new IMAGE_PLUGIN_BRUSH("Icons/ModelingUVSeamEdit_x40",			Icon20x20));
+		StyleSet->Set("ModelingToolsManagerCommands.BeginUVSeamEditTool.Small", 		new IMAGE_PLUGIN_BRUSH("Icons/ModelingUVSeamEdit_x40",			Icon20x20));
+		StyleSet->Set("ModelingToolsManagerCommands.BeginBakeMeshAttributeMapsTool", 			new IMAGE_PLUGIN_BRUSH("Icons/ModelingBakeMaps_x40",			Icon20x20));
+		StyleSet->Set("ModelingToolsManagerCommands.BeginBakeMeshAttributeMapsTool.Small", 		new IMAGE_PLUGIN_BRUSH("Icons/ModelingBakeMaps_x40",			Icon20x20));
+		StyleSet->Set("ModelingToolsManagerCommands.BeginBakeMeshAttributeVertexTool", new IMAGE_PLUGIN_BRUSH("Icons/ModelingBakeMaps_x40", Icon20x20));
+		StyleSet->Set("ModelingToolsManagerCommands.BeginBakeMeshAttributeVertexTool.Small", new IMAGE_PLUGIN_BRUSH("Icons/ModelingBakeMaps_x40", Icon20x20));
+		//StyleSet->Set("ModelingToolsManagerCommands.BeginRemoveOccludedTrianglesTool", 				new IMAGE_PLUGIN_BRUSH("Icons/Jacket_40x",			Icon20x20));
+		//StyleSet->Set("ModelingToolsManagerCommands.BeginRemoveOccludedTrianglesTool.Small", 		new IMAGE_PLUGIN_BRUSH("Icons/Jacket_40x",			Icon20x20));
+		StyleSet->Set("ModelingToolsManagerCommands.BeginHoleFillTool", 				new IMAGE_PLUGIN_BRUSH("Icons/ModelingHoleFill_x40",			Icon20x20));
+		StyleSet->Set("ModelingToolsManagerCommands.BeginHoleFillTool.Small", 		new IMAGE_PLUGIN_BRUSH("Icons/ModelingHoleFill_x40",			Icon20x20));
+		StyleSet->Set("ModelingToolsManagerCommands.BeginUVProjectionTool", 			new IMAGE_PLUGIN_BRUSH("Icons/UVProjection_40x", 	Icon20x20));
+		StyleSet->Set("ModelingToolsManagerCommands.BeginUVProjectionTool.Small", 		new IMAGE_PLUGIN_BRUSH("Icons/UVProjection_40x", 	Icon20x20));
+		//StyleSet->Set("ModelingToolsManagerCommands.BeginUVLayoutTool", 			new IMAGE_PLUGIN_BRUSH("Icons/UVLayout_40x", 	Icon20x20));
+		//StyleSet->Set("ModelingToolsManagerCommands.BeginUVLayoutTool.Small", 		new IMAGE_PLUGIN_BRUSH("Icons/UVLayout_40x", 	Icon20x20));
+		StyleSet->Set("ModelingToolsManagerCommands.BeginVoxelMergeTool", 				new IMAGE_PLUGIN_BRUSH("Icons/VoxMerge_40x", 		Icon20x20));
+		StyleSet->Set("ModelingToolsManagerCommands.BeginVoxelMergeTool.Small", 		new IMAGE_PLUGIN_BRUSH("Icons/VoxMerge_40x", 		Icon20x20));
+		StyleSet->Set("ModelingToolsManagerCommands.BeginVoxelBooleanTool", 			new IMAGE_PLUGIN_BRUSH("Icons/VoxBoolean_40x", 		Icon20x20));
+		StyleSet->Set("ModelingToolsManagerCommands.BeginVoxelBooleanTool.Small", 		new IMAGE_PLUGIN_BRUSH("Icons/VoxBoolean_40x", 		Icon20x20));
+		//StyleSet->Set("ModelingToolsManagerCommands.BeginSelfUnionTool",				new IMAGE_PLUGIN_BRUSH("Icons/MeshMerge_40x",		Icon20x20));
+		//StyleSet->Set("ModelingToolsManagerCommands.BeginSelfUnionTool.Small",			new IMAGE_PLUGIN_BRUSH("Icons/MeshMerge_40x",		Icon20x20));
+		//StyleSet->Set("ModelingToolsManagerCommands.BeginMeshBooleanTool",				new IMAGE_PLUGIN_BRUSH("Icons/Boolean_40x",			Icon20x20));
+		//StyleSet->Set("ModelingToolsManagerCommands.BeginMeshBooleanTool.Small",		new IMAGE_PLUGIN_BRUSH("Icons/Boolean_40x",			Icon20x20));
+		StyleSet->Set("ModelingToolsManagerCommands.BeginPlaneCutTool", 				new IMAGE_PLUGIN_BRUSH("Icons/PlaneCut_40x", 		Icon20x20));
+		StyleSet->Set("ModelingToolsManagerCommands.BeginPlaneCutTool.Small", 			new IMAGE_PLUGIN_BRUSH("Icons/PlaneCut_40x", 		Icon20x20));
+		StyleSet->Set("ModelingToolsManagerCommands.BeginMirrorTool", 				    new IMAGE_PLUGIN_BRUSH("Icons/ModelingMirror_x40", 		Icon20x20));
+		StyleSet->Set("ModelingToolsManagerCommands.BeginMirrorTool.Small", 			new IMAGE_PLUGIN_BRUSH("Icons/ModelingMirror_x40", 		Icon20x20));
+		StyleSet->Set("ModelingToolsManagerCommands.BeginOffsetMeshTool", 				new IMAGE_PLUGIN_BRUSH("Icons/ModelingOffset_x40", 		Icon20x20));
+		StyleSet->Set("ModelingToolsManagerCommands.BeginOffsetMeshTool.Small", 		new IMAGE_PLUGIN_BRUSH("Icons/ModelingOffset_x40", 		Icon20x20));
+		StyleSet->Set("ModelingToolsManagerCommands.BeginDisplaceMeshTool", 			new IMAGE_PLUGIN_BRUSH("Icons/ModelingDisplace_x40", 		Icon20x20));
+		StyleSet->Set("ModelingToolsManagerCommands.BeginDisplaceMeshTool.Small", 		new IMAGE_PLUGIN_BRUSH("Icons/ModelingDisplace_x40", 		Icon20x20));
+		StyleSet->Set("ModelingToolsManagerCommands.BeginMeshSelectionTool", 			new IMAGE_PLUGIN_BRUSH("Icons/MeshSelect_40x",		Icon20x20));
+		StyleSet->Set("ModelingToolsManagerCommands.BeginMeshSelectionTool.Small", 	new IMAGE_PLUGIN_BRUSH("Icons/MeshSelect_40x",		Icon20x20));
+
+		StyleSet->Set("ModelingToolsManagerCommands.BeginMeshInspectorTool", 			new IMAGE_PLUGIN_BRUSH("Icons/Inspector_40x",		Icon20x20));
+		StyleSet->Set("ModelingToolsManagerCommands.BeginMeshInspectorTool.Small", 		new IMAGE_PLUGIN_BRUSH("Icons/Inspector_40x",		Icon20x20));
+		StyleSet->Set("ModelingToolsManagerCommands.BeginWeldEdgesTool", 				new IMAGE_PLUGIN_BRUSH("Icons/WeldEdges_40x", 		Icon20x20));
+		StyleSet->Set("ModelingToolsManagerCommands.BeginWeldEdgesTool.Small", 			new IMAGE_PLUGIN_BRUSH("Icons/WeldEdges_40x", 		Icon20x20));
+		StyleSet->Set("ModelingToolsManagerCommands.BeginAttributeEditorTool", 			new IMAGE_PLUGIN_BRUSH("Icons/AttributeEditor_40x", Icon20x20));
+		StyleSet->Set("ModelingToolsManagerCommands.BeginAttributeEditorTool.Small", 	new IMAGE_PLUGIN_BRUSH("Icons/AttributeEditor_40x", Icon20x20));
+
+		StyleSet->Set("ModelingToolsManagerCommands.BeginAlignObjectsTool",                  new FSlateImageBrush(StyleSet->RootToCoreContentDir(TEXT("../Editor/Slate/Icons/GeneralTools/Align_40x.png")), Icon20x20));
+		StyleSet->Set("ModelingToolsManagerCommands.BeginAlignObjectsTool.Small",            new FSlateImageBrush(StyleSet->RootToCoreContentDir(TEXT("../Editor/Slate/Icons/GeneralTools/Align_40x.png")), Icon20x20));
+		StyleSet->Set("ModelingToolsManagerCommands.BeginTransferMeshTool",                  new FSlateImageBrush(StyleSet->RootToCoreContentDir(TEXT("../Editor/Slate/Icons/GeneralTools/Next_40x.png")), Icon20x20));
+		StyleSet->Set("ModelingToolsManagerCommands.BeginTransferMeshTool.Small",            new FSlateImageBrush(StyleSet->RootToCoreContentDir(TEXT("../Editor/Slate/Icons/GeneralTools/Next_40x.png")), Icon20x20));
+		StyleSet->Set("ModelingToolsManagerCommands.BeginGlobalUVGenerateTool",              new IMAGE_PLUGIN_BRUSH("Icons/AutoUnwrap_40x",       Icon20x20));      
+		StyleSet->Set("ModelingToolsManagerCommands.BeginGlobalUVGenerateTool.Small",        new IMAGE_PLUGIN_BRUSH("Icons/AutoUnwrap_40x",       Icon20x20));      
+		StyleSet->Set("ModelingToolsManagerCommands.BeginBakeTransformTool",                 new IMAGE_PLUGIN_BRUSH("Icons/BakeXForm_40x",        Icon20x20));      
+		StyleSet->Set("ModelingToolsManagerCommands.BeginBakeTransformTool.Small",           new IMAGE_PLUGIN_BRUSH("Icons/BakeXForm_40x",        Icon20x20));      
+		StyleSet->Set("ModelingToolsManagerCommands.BeginCombineMeshesTool",                 new IMAGE_BRUSH_SVG("Icons/GeometryCombine",          Icon20x20));
+		StyleSet->Set("ModelingToolsManagerCommands.BeginCombineMeshesTool.Small",           new IMAGE_BRUSH_SVG("Icons/GeometryCombine",          Icon20x20));
+		StyleSet->Set("ModelingToolsManagerCommands.BeginDuplicateMeshesTool",               new IMAGE_PLUGIN_BRUSH("Icons/Duplicate_40x",        Icon20x20));   
+		StyleSet->Set("ModelingToolsManagerCommands.BeginDuplicateMeshesTool.Small",         new IMAGE_PLUGIN_BRUSH("Icons/Duplicate_40x",        Icon20x20));   
+		StyleSet->Set("ModelingToolsManagerCommands.BeginEditMeshMaterialsTool",             new IMAGE_PLUGIN_BRUSH("Icons/EditMats_40x",         Icon20x20));     
+		StyleSet->Set("ModelingToolsManagerCommands.BeginEditMeshMaterialsTool.Small",       new IMAGE_PLUGIN_BRUSH("Icons/EditMats_40x",         Icon20x20));     
+		StyleSet->Set("ModelingToolsManagerCommands.BeginEditPivotTool",                     new IMAGE_PLUGIN_BRUSH("Icons/EditPivot_40x",        Icon20x20));      
+		StyleSet->Set("ModelingToolsManagerCommands.BeginEditPivotTool.Small",               new IMAGE_PLUGIN_BRUSH("Icons/EditPivot_40x",        Icon20x20));      
+		StyleSet->Set("ModelingToolsManagerCommands.BeginGroupUVGenerateTool",               new IMAGE_PLUGIN_BRUSH("Icons/GroupUnwrap_40x",      Icon20x20));       
+		StyleSet->Set("ModelingToolsManagerCommands.BeginGroupUVGenerateTool.Small",         new IMAGE_PLUGIN_BRUSH("Icons/GroupUnwrap_40x",      Icon20x20));       
+		StyleSet->Set("ModelingToolsManagerCommands.BeginRemoveOccludedTrianglesTool",       new IMAGE_PLUGIN_BRUSH("Icons/Jacketing_40x",        Icon20x20));     
+		StyleSet->Set("ModelingToolsManagerCommands.BeginRemoveOccludedTrianglesTool.Small", new IMAGE_PLUGIN_BRUSH("Icons/Jacketing_40x",        Icon20x20));     
+		StyleSet->Set("ModelingToolsManagerCommands.BeginPolygonCutTool",                    new IMAGE_PLUGIN_BRUSH("Icons/PolyCut_40x",          Icon20x20));   
+		StyleSet->Set("ModelingToolsManagerCommands.BeginPolygonCutTool.Small",              new IMAGE_PLUGIN_BRUSH("Icons/PolyCut_40x",          Icon20x20));   
+		StyleSet->Set("ModelingToolsManagerCommands.BeginPolyDeformTool",                    new IMAGE_PLUGIN_BRUSH("Icons/PolyDeform_40x",       Icon20x20));      
+		StyleSet->Set("ModelingToolsManagerCommands.BeginPolyDeformTool.Small",              new IMAGE_PLUGIN_BRUSH("Icons/PolyDeform_40x",       Icon20x20));      
+		StyleSet->Set("ModelingToolsManagerCommands.BeginPolyGroupsTool",                    new IMAGE_PLUGIN_BRUSH("Icons/PolyGroups_40x",       Icon20x20));      
+		StyleSet->Set("ModelingToolsManagerCommands.BeginPolyGroupsTool.Small",              new IMAGE_PLUGIN_BRUSH("Icons/PolyGroups_40x",       Icon20x20));      
+		StyleSet->Set("ModelingToolsManagerCommands.BeginDrawPolyPathTool",                  new IMAGE_PLUGIN_BRUSH("Icons/PolyPath_40x",         Icon20x20));    
+		StyleSet->Set("ModelingToolsManagerCommands.BeginDrawPolyPathTool.Small",            new IMAGE_PLUGIN_BRUSH("Icons/PolyPath_40x",         Icon20x20));    
+		StyleSet->Set("ModelingToolsManagerCommands.BeginDrawAndRevolveTool",                new IMAGE_PLUGIN_BRUSH("Icons/ModelingDrawAndRevolve_x40", Icon20x20));
+		StyleSet->Set("ModelingToolsManagerCommands.BeginDrawAndRevolveTool.Small",          new IMAGE_PLUGIN_BRUSH("Icons/ModelingDrawAndRevolve_x20", Icon20x20));
+		StyleSet->Set("ModelingToolsManagerCommands.BeginRevolveBoundaryTool",               new IMAGE_PLUGIN_BRUSH("Icons/ModelingRevolveBoundary_x40", Icon20x20));
+		StyleSet->Set("ModelingToolsManagerCommands.BeginRevolveBoundaryTool.Small",         new IMAGE_PLUGIN_BRUSH("Icons/ModelingRevolveBoundary_x20", Icon20x20));
+		StyleSet->Set("ModelingToolsManagerCommands.BeginMeshBooleanTool",                   new IMAGE_PLUGIN_BRUSH("Icons/ModelingMeshBoolean_x40", Icon20x20));
+		StyleSet->Set("ModelingToolsManagerCommands.BeginMeshBooleanTool.Small",             new IMAGE_PLUGIN_BRUSH("Icons/ModelingMeshBoolean_x20", Icon20x20));
+		StyleSet->Set("ModelingToolsManagerCommands.BeginMeshTrimTool",				 new IMAGE_BRUSH_SVG("Icons/ModelingTrim", Icon20x20));
+		StyleSet->Set("ModelingToolsManagerCommands.BeginMeshTrimTool.Small",			 new IMAGE_BRUSH_SVG("Icons/ModelingTrim", Icon20x20));
+		StyleSet->Set("ModelingToolsManagerCommands.BeginCutMeshWithMeshTool",               new IMAGE_BRUSH_SVG("Icons/ModelingMeshCut", Icon20x20));
+		StyleSet->Set("ModelingToolsManagerCommands.BeginCutMeshWithMeshTool.Small",         new IMAGE_BRUSH_SVG("Icons/ModelingMeshCut", Icon20x20));
+		StyleSet->Set("ModelingToolsManagerCommands.BeginSelfUnionTool",                     new IMAGE_PLUGIN_BRUSH("Icons/ModelingSelfUnion_x40", Icon20x20));
+		StyleSet->Set("ModelingToolsManagerCommands.BeginSelfUnionTool.Small",               new IMAGE_PLUGIN_BRUSH("Icons/ModelingSelfUnion_x20", Icon20x20));
+		StyleSet->Set("ModelingToolsManagerCommands.BeginVoxelSolidifyTool",                 new IMAGE_PLUGIN_BRUSH("Icons/ModelingVoxSolidify_x40", Icon20x20));
+		StyleSet->Set("ModelingToolsManagerCommands.BeginVoxelSolidifyTool.Small",           new IMAGE_PLUGIN_BRUSH("Icons/ModelingVoxSolidify_x20", Icon20x20));
+		StyleSet->Set("ModelingToolsManagerCommands.BeginVoxelBlendTool",                    new IMAGE_PLUGIN_BRUSH("Icons/ModelingVoxBlend_x40", Icon20x20));
+		StyleSet->Set("ModelingToolsManagerCommands.BeginVoxelBlendTool.Small",              new IMAGE_PLUGIN_BRUSH("Icons/ModelingVoxBlend_x20", Icon20x20));
+		StyleSet->Set("ModelingToolsManagerCommands.BeginVoxelMorphologyTool",               new IMAGE_PLUGIN_BRUSH("Icons/ModelingVoxMorphology_x40", Icon20x20));
+		StyleSet->Set("ModelingToolsManagerCommands.BeginVoxelMorphologyTool.Small",         new IMAGE_PLUGIN_BRUSH("Icons/ModelingVoxMorphology_x20", Icon20x20));
+		StyleSet->Set("ModelingToolsManagerCommands.BeginMeshSpaceDeformerTool",             new IMAGE_PLUGIN_BRUSH("Icons/SpaceDeform_40x",      Icon20x20));       
+		StyleSet->Set("ModelingToolsManagerCommands.BeginMeshSpaceDeformerTool.Small",       new IMAGE_PLUGIN_BRUSH("Icons/SpaceDeform_40x",      Icon20x20));       
+		StyleSet->Set("ModelingToolsManagerCommands.BeginMeshAttributePaintTool",             new IMAGE_PLUGIN_BRUSH("Icons/ModelingAttributePaint_x40",      Icon20x20));
+		StyleSet->Set("ModelingToolsManagerCommands.BeginMeshAttributePaintTool.Small",       new IMAGE_PLUGIN_BRUSH("Icons/ModelingAttributePaint_x40",      Icon20x20));
+		StyleSet->Set("ModelingToolsManagerCommands.BeginTransformUVIslandsTool",            new IMAGE_PLUGIN_BRUSH("Icons/TransformUVs_40x",     Icon20x20));         
+		StyleSet->Set("ModelingToolsManagerCommands.BeginTransformUVIslandsTool.Small",      new IMAGE_PLUGIN_BRUSH("Icons/TransformUVs_40x",     Icon20x20));         
+		StyleSet->Set("ModelingToolsManagerCommands.BeginUVLayoutTool",                      new IMAGE_PLUGIN_BRUSH("Icons/UVLayout_40x",         Icon20x20));    
+		StyleSet->Set("ModelingToolsManagerCommands.BeginUVLayoutTool.Small",                new IMAGE_PLUGIN_BRUSH("Icons/UVLayout_40x",         Icon20x20));    
+		StyleSet->Set("ModelingToolsManagerCommands.BeginMeshGroupPaintTool", new IMAGE_BRUSH_SVG("Icons/GroupPaint", Icon20x20));
+		StyleSet->Set("ModelingToolsManagerCommands.BeginMeshGroupPaintTool.Small", new IMAGE_BRUSH_SVG("Icons/GroupPaint", Icon20x20));
+		StyleSet->Set("ModelingToolsManagerCommands.BeginLatticeDeformerTool", new IMAGE_BRUSH_SVG("Icons/LatticeDeformation", Icon20x20));
+		StyleSet->Set("ModelingToolsManagerCommands.BeginLatticeDeformerTool.Small", new IMAGE_BRUSH_SVG("Icons/LatticeDeformation", Icon20x20));
+
+		
+
+		StyleSet->Set("ModelingToolsManagerCommands.BeginVolumeToMeshTool",                  new IMAGE_PLUGIN_BRUSH("Icons/ModelingVol2Mesh_x40",         Icon20x20));
+		StyleSet->Set("ModelingToolsManagerCommands.BeginVolumeToMeshTool.Small",            new IMAGE_PLUGIN_BRUSH("Icons/ModelingVol2Mesh_x40",         Icon20x20));
+		StyleSet->Set("ModelingToolsManagerCommands.BeginMeshToVolumeTool",                  new IMAGE_PLUGIN_BRUSH("Icons/ModelingMesh2Vol_x40",         Icon20x20));
+		StyleSet->Set("ModelingToolsManagerCommands.BeginMeshToVolumeTool.Small",            new IMAGE_PLUGIN_BRUSH("Icons/ModelingMesh2Vol_x40",         Icon20x20));
+		StyleSet->Set("ModelingToolsManagerCommands.BeginBspConversionTool",                 new IMAGE_PLUGIN_BRUSH("Icons/ModelingBSPConversion_x40",         Icon20x20));
+		StyleSet->Set("ModelingToolsManagerCommands.BeginBspConversionTool.Small",           new IMAGE_PLUGIN_BRUSH("Icons/ModelingBSPConversion_x40",         Icon20x20));
+		StyleSet->Set("ModelingToolsManagerCommands.BeginPhysicsInspectorTool",              new IMAGE_PLUGIN_BRUSH("Icons/ModelingPhysInspect_x40",         Icon20x20));
+		StyleSet->Set("ModelingToolsManagerCommands.BeginPhysicsInspectorTool.Small",        new IMAGE_PLUGIN_BRUSH("Icons/ModelingPhysInspect_x40",         Icon20x20));
+		StyleSet->Set("ModelingToolsManagerCommands.BeginSetCollisionGeometryTool",          new IMAGE_PLUGIN_BRUSH("Icons/ModelingMeshToCollision_x40",         Icon20x20));
+		StyleSet->Set("ModelingToolsManagerCommands.BeginSetCollisionGeometryTool.Small",    new IMAGE_PLUGIN_BRUSH("Icons/ModelingMeshToCollision_x40",         Icon20x20));
+		StyleSet->Set("ModelingToolsManagerCommands.BeginExtractCollisionGeometryTool",      new IMAGE_PLUGIN_BRUSH("Icons/ModelingCollisionToMesh_x40",         Icon20x20));
+		StyleSet->Set("ModelingToolsManagerCommands.BeginExtractCollisionGeometryTool.Small",new IMAGE_PLUGIN_BRUSH("Icons/ModelingCollisionToMesh_x40",         Icon20x20));
+
+		StyleSet->Set("ModelingToolsManagerCommands.BeginGenerateStaticMeshLODAssetTool", new IMAGE_BRUSH_SVG("Icons/AutoLOD", Icon20x20));
+		StyleSet->Set("ModelingToolsManagerCommands.BeginGenerateStaticMeshLODAssetTool.Small", new IMAGE_BRUSH_SVG("Icons/AutoLOD", Icon20x20));
+		StyleSet->Set("ModelingToolsManagerCommands.BeginLODManagerTool", new IMAGE_BRUSH_SVG("Icons/ModelingLODManager", Icon20x20));
+		StyleSet->Set("ModelingToolsManagerCommands.BeginLODManagerTool.Small", new IMAGE_BRUSH_SVG("Icons/ModelingLODManager", Icon20x20));
+
+		StyleSet->Set("ModelingToolsManagerCommands.BeginGroomCardsEditorTool", new IMAGE_BRUSH_SVG("Icons/CardsEditor", Icon20x20));
+		StyleSet->Set("ModelingToolsManagerCommands.BeginGroomCardsEditorTool.Small", new IMAGE_BRUSH_SVG("Icons/CardsEditor", Icon20x20));
+		StyleSet->Set("ModelingToolsManagerCommands.BeginGenerateLODMeshesTool", new IMAGE_BRUSH_SVG("Icons/GenLODs", Icon20x20));
+		StyleSet->Set("ModelingToolsManagerCommands.BeginGenerateLODMeshesTool.Small", new IMAGE_BRUSH_SVG("Icons/GenLODs", Icon20x20));
+		StyleSet->Set("ModelingToolsManagerCommands.BeginGroomToMeshTool", new IMAGE_BRUSH_SVG("Icons/HairHelmet", Icon20x20));
+		StyleSet->Set("ModelingToolsManagerCommands.BeginGroomToMeshTool.Small", new IMAGE_BRUSH_SVG("Icons/HairHelmet", Icon20x20));
+
+		StyleSet->Set("ModelingToolsManagerCommands.BeginSkinWeightsPaintTool", new IMAGE_BRUSH_SVG("Icons/SkinWeightsPaint", Icon20x20));
+		StyleSet->Set("ModelingToolsManagerCommands.BeginSkinWeightsPaintTool.Small", new IMAGE_BRUSH_SVG("Icons/SkinWeightsPaint", Icon20x20));
+	}
+
+	FSlateStyleRegistry::RegisterSlateStyle(*StyleSet.Get());
+};
+
+END_SLATE_FUNCTION_BUILD_OPTIMIZATION
+
+#undef IMAGE_PLUGIN_BRUSH
+#undef IMAGE_BRUSH
+#undef BOX_BRUSH
+#undef DEFAULT_FONT
+
+void FModelingToolsEditorModeStyle::Shutdown()
+{
+	if (StyleSet.IsValid())
+	{
+		FSlateStyleRegistry::UnRegisterSlateStyle(*StyleSet.Get());
+		ensure(StyleSet.IsUnique());
+		StyleSet.Reset();
+	}
+}
diff --git a/Engine/Plugins/Editor/ModelingToolsEditorMode/Source/ModelingToolsEditorMode/Private/ModelingToolsEditorModeToolkit.cpp b/Engine/Plugins/Editor/ModelingToolsEditorMode/Source/ModelingToolsEditorMode/Private/ModelingToolsEditorModeToolkit.cpp
new file mode 100644
--- /dev/null
+++ b/Engine/Plugins/Editor/ModelingToolsEditorMode/Source/ModelingToolsEditorMode/Private/ModelingToolsEditorModeToolkit.cpp
@@ -0,0 +1,922 @@
+// Copyright Epic Games, Inc. All Rights Reserved.
+
+#include "ModelingToolsEditorModeToolkit.h"
+#include "ModelingToolsEditorMode.h"
+#include "ModelingToolsManagerActions.h"
+#include "ModelingToolsEditorModeSettings.h"
+#include "ModelingToolsEditorModeStyle.h"
+#include "Engine/Selection.h"
+
+#include "Framework/MultiBox/MultiBoxBuilder.h"
+#include "Modules/ModuleManager.h"
+#include "PropertyEditorModule.h"
+#include "IDetailsView.h"
+#include "IDetailRootObjectCustomization.h"
+#include "ISettingsModule.h"
+#include "EditorModeManager.h"
+
+#include "Widgets/Input/SButton.h"
+#include "Widgets/Text/STextBlock.h"
+#include "Widgets/Layout/SSeparator.h"
+#include "Widgets/Images/SImage.h"
+#include "Widgets/Layout/SExpandableArea.h"
+
+// for Tool Extensions
+#include "Features/IModularFeatures.h"
+#include "ModelingModeToolExtensions.h"
+
+// for Object Type properties
+#include "PropertySets/CreateMeshObjectTypeProperties.h"
+
+// for LOD setting
+#include "EditorInteractiveToolsFrameworkModule.h"
+#include "Tools/EditorComponentSourceFactory.h"
+#include "ToolTargetManager.h"
+#include "ToolTargets/StaticMeshComponentToolTarget.h"
+
+
+#define LOCTEXT_NAMESPACE "FModelingToolsEditorModeToolkit"
+
+
+// if set to 1, then on mode initialization we include buttons for prototype modeling tools
+static TAutoConsoleVariable<int32> CVarEnablePrototypeModelingTools(
+	TEXT("modeling.EnablePrototypes"),
+	0,
+	TEXT("Enable unsupported Experimental prototype Modeling Tools"));
+static TAutoConsoleVariable<int32> CVarEnablePolyModeling(
+	TEXT("modeling.EnablePolyModel"),
+	0,
+	TEXT("Enable prototype PolyEdit tab"));
+
+
+
+FModelingToolsEditorModeToolkit::FModelingToolsEditorModeToolkit()
+{
+}
+
+
+FModelingToolsEditorModeToolkit::~FModelingToolsEditorModeToolkit()
+{
+	UModelingToolsEditorModeSettings* Settings = GetMutableDefault<UModelingToolsEditorModeSettings>();
+	Settings->OnModified.Remove(AssetSettingsModifiedHandle);
+	GetScriptableEditorMode()->GetInteractiveToolsContext()->OnToolNotificationMessage.RemoveAll(this);
+	GetScriptableEditorMode()->GetInteractiveToolsContext()->OnToolWarningMessage.RemoveAll(this);
+}
+
+
+void FModelingToolsEditorModeToolkit::CustomizeModeDetailsViewArgs(FDetailsViewArgs& ArgsInOut)
+{
+	//ArgsInOut.ColumnWidth = 0.3f;
+}
+
+void FModelingToolsEditorModeToolkit::Init(const TSharedPtr<IToolkitHost>& InitToolkitHost, TWeakObjectPtr<UEdMode> InOwningMode)
+{
+	// Have to create the ToolkitWidget here because FModeToolkit::Init() is going to ask for it and add
+	// it to the Mode panel, and not ask again afterwards. However we have to call Init() to get the 
+	// ModeDetailsView created, that we need to add to the ToolkitWidget. So, we will create the Widget
+	// here but only add the rows to it after we call Init()
+	TSharedPtr<SVerticalBox> ToolkitWidgetVBox = SNew(SVerticalBox);
+	SAssignNew(ToolkitWidget, SBorder)
+		.HAlign(HAlign_Fill)
+		.Padding(4)
+		[
+			ToolkitWidgetVBox->AsShared()
+		];
+
+	FModeToolkit::Init(InitToolkitHost, InOwningMode);
+
+	GetToolkitHost()->OnActiveViewportChanged().AddSP(this, &FModelingToolsEditorModeToolkit::OnActiveViewportChanged);
+
+	ModeWarningArea = SNew(STextBlock)
+		.AutoWrapText(true)
+		.Font(FCoreStyle::GetDefaultFontStyle("Bold", 9))
+		.ColorAndOpacity(FSlateColor(FLinearColor(0.9f, 0.15f, 0.15f)));
+	ModeWarningArea->SetText(FText::GetEmpty());
+	ModeWarningArea->SetVisibility(EVisibility::Collapsed);
+
+	ModeHeaderArea = SNew(STextBlock)
+		.AutoWrapText(true)
+		.Font(FCoreStyle::GetDefaultFontStyle("Bold", 12));
+	ModeHeaderArea->SetText(LOCTEXT("SelectToolLabel", "Select a Tool from the Toolbar"));
+	ModeHeaderArea->SetJustification(ETextJustify::Center);
+
+
+	ToolWarningArea = SNew(STextBlock)
+		.AutoWrapText(true)
+		.Font(FCoreStyle::GetDefaultFontStyle("Bold", 9))
+		.ColorAndOpacity(FSlateColor(FLinearColor(0.9f, 0.15f, 0.15f)));
+	ToolWarningArea->SetText(FText::GetEmpty());
+
+	// add the various sections to the mode toolbox
+	ToolkitWidgetVBox->AddSlot().AutoHeight().HAlign(HAlign_Fill).Padding(5)
+		[
+			ModeWarningArea->AsShared()
+		];
+	ToolkitWidgetVBox->AddSlot().AutoHeight().HAlign(HAlign_Fill).Padding(5)
+		[
+			ModeHeaderArea->AsShared()
+		];
+	ToolkitWidgetVBox->AddSlot().AutoHeight().HAlign(HAlign_Fill).Padding(5)
+		[
+			ToolWarningArea->AsShared()
+		];
+	ToolkitWidgetVBox->AddSlot().HAlign(HAlign_Fill).FillHeight(1.f)
+		[
+			ModeDetailsView->AsShared()
+		];
+	ToolkitWidgetVBox->AddSlot().AutoHeight().HAlign(HAlign_Fill).Padding(5)
+		[
+			MakeAssetConfigPanel()->AsShared()
+		];
+
+	ClearNotification();
+	ClearWarning();
+
+	ActiveToolName = FText::GetEmpty();
+	ActiveToolMessage = FText::GetEmpty();
+
+	GetScriptableEditorMode()->GetInteractiveToolsContext()->OnToolNotificationMessage.AddSP(this, &FModelingToolsEditorModeToolkit::PostNotification);
+	GetScriptableEditorMode()->GetInteractiveToolsContext()->OnToolWarningMessage.AddSP(this, &FModelingToolsEditorModeToolkit::PostWarning);
+
+	UpdateObjectCreationOptionsFromSettings();
+
+	SAssignNew(ViewportOverlayWidget, SHorizontalBox)
+
+	+SHorizontalBox::Slot()
+	.HAlign(HAlign_Center)
+	.VAlign(VAlign_Bottom)
+	.Padding(FMargin(0.0f, 0.0f, 0.f, 15.f))
+	[
+		SNew(SBorder)
+		.BorderImage(FAppStyle::Get().GetBrush("EditorViewport.OverlayBrush"))
+		.Padding(8.f)
+		[
+			SNew(SHorizontalBox)
+
+			+SHorizontalBox::Slot()
+			.AutoWidth()
+			.VAlign(VAlign_Center)
+			.Padding(FMargin(0.f, 0.f, 8.f, 0.f))
+			[
+				SNew(SImage)
+				.Image_Lambda([this] () { return ActiveToolIcon; })
+			]
+
+			+SHorizontalBox::Slot()
+			.AutoWidth()
+			.VAlign(VAlign_Center)
+			.Padding(FMargin(0.f, 0.f, 8.f, 0.f))
+			[
+				SNew(STextBlock)
+				.Text(this, &FModelingToolsEditorModeToolkit::GetActiveToolDisplayName)
+			]
+
+			+SHorizontalBox::Slot()
+			.AutoWidth()
+			.Padding(FMargin(0.0, 0.f, 2.f, 0.f))
+			[
+				SNew(SButton)
+				.ButtonStyle(FAppStyle::Get(), "PrimaryButton")
+				.TextStyle( FAppStyle::Get(), "DialogButtonText" )
+				.Text(LOCTEXT("OverlayAccept", "Accept"))
+				.ToolTipText(LOCTEXT("OverlayAcceptTooltip", "Accept/Commit the results of the active Tool [Enter]"))
+				.HAlign(HAlign_Center)
+				.OnClicked_Lambda([this]() { GetScriptableEditorMode()->GetInteractiveToolsContext()->EndTool(EToolShutdownType::Accept); return FReply::Handled(); })
+				.IsEnabled_Lambda([this]() { return GetScriptableEditorMode()->GetInteractiveToolsContext()->CanAcceptActiveTool(); })
+				.Visibility_Lambda([this]() { return GetScriptableEditorMode()->GetInteractiveToolsContext()->ActiveToolHasAccept() ? EVisibility::Visible : EVisibility::Collapsed; })
+			]
+
+			+SHorizontalBox::Slot()
+			.AutoWidth()
+			.Padding(FMargin(2.0, 0.f, 0.f, 0.f))
+			[
+				SNew(SButton)
+				.ButtonStyle(FAppStyle::Get(), "Button")
+				.TextStyle( FAppStyle::Get(), "DialogButtonText" )
+				.Text(LOCTEXT("OverlayCancel", "Cancel"))
+				.ToolTipText(LOCTEXT("OverlayCancelTooltip", "Cancel the active Tool [Esc]"))
+				.HAlign(HAlign_Center)
+				.OnClicked_Lambda([this]() { GetScriptableEditorMode()->GetInteractiveToolsContext()->EndTool(EToolShutdownType::Cancel); return FReply::Handled(); })
+				.IsEnabled_Lambda([this]() { return GetScriptableEditorMode()->GetInteractiveToolsContext()->CanCancelActiveTool(); })
+				.Visibility_Lambda([this]() { return GetScriptableEditorMode()->GetInteractiveToolsContext()->ActiveToolHasAccept() ? EVisibility::Visible : EVisibility::Collapsed; })
+			]
+
+			+SHorizontalBox::Slot()
+			.AutoWidth()
+			.Padding(FMargin(2.0, 0.f, 0.f, 0.f))
+			[
+				SNew(SButton)
+				.ButtonStyle(FAppStyle::Get(), "PrimaryButton")
+				.TextStyle( FAppStyle::Get(), "DialogButtonText" )
+				.Text(LOCTEXT("OverlayComplete", "Complete"))
+				.ToolTipText(LOCTEXT("OverlayCompleteTooltip", "Exit the active Tool [Enter]"))
+				.HAlign(HAlign_Center)
+				.OnClicked_Lambda([this]() { GetScriptableEditorMode()->GetInteractiveToolsContext()->EndTool(EToolShutdownType::Completed); return FReply::Handled(); })
+				.IsEnabled_Lambda([this]() { return GetScriptableEditorMode()->GetInteractiveToolsContext()->CanCompleteActiveTool(); })
+				.Visibility_Lambda([this]() { return GetScriptableEditorMode()->GetInteractiveToolsContext()->CanCompleteActiveTool() ? EVisibility::Visible : EVisibility::Collapsed; })
+			]
+		]	
+	];
+
+}
+
+
+
+TSharedPtr<SWidget> FModelingToolsEditorModeToolkit::MakeAssetConfigPanel()
+{
+	//
+	// New Asset Location drop-down
+	//
+
+	AssetLocationModes.Add(MakeShared<FString>(TEXT("AutoGen Folder (World-Relative)")));
+	AssetLocationModes.Add(MakeShared<FString>(TEXT("AutoGen Folder (Global)")));
+	AssetLocationModes.Add(MakeShared<FString>(TEXT("Current Folder")));
+	AssetLocationMode = SNew(STextComboBox)
+		.OptionsSource(&AssetLocationModes)
+		.OnSelectionChanged_Lambda([&](TSharedPtr<FString> String, ESelectInfo::Type) { UpdateAssetLocationMode(String); });
+	AssetSaveModes.Add(MakeShared<FString>(TEXT("AutoSave New Assets")));
+	AssetSaveModes.Add(MakeShared<FString>(TEXT("Manual Save")));
+	AssetSaveModes.Add(MakeShared<FString>(TEXT("Interactive")));
+	AssetSaveMode = SNew(STextComboBox)
+		.OptionsSource(&AssetSaveModes)
+		.OnSelectionChanged_Lambda([&](TSharedPtr<FString> String, ESelectInfo::Type) { UpdateAssetSaveMode(String); });
+	
+	// initialize combos
+	UpdateAssetPanelFromSettings();
+	
+	// register callback
+	UModelingToolsEditorModeSettings* Settings = GetMutableDefault<UModelingToolsEditorModeSettings>();
+	AssetSettingsModifiedHandle = Settings->OnModified.AddLambda([this](UObject*, FProperty*) { OnAssetSettingsModified(); });
+
+
+	//
+	// LOD selection dropdown
+	//
+
+	AssetLODModes.Add(MakeShared<FString>(TEXT("Max Available")));
+	AssetLODModes.Add(MakeShared<FString>(TEXT("HiRes")));
+	AssetLODModes.Add(MakeShared<FString>(TEXT("LOD0")));
+	AssetLODModes.Add(MakeShared<FString>(TEXT("LOD1")));
+	AssetLODModes.Add(MakeShared<FString>(TEXT("LOD2")));
+	AssetLODModes.Add(MakeShared<FString>(TEXT("LOD3")));
+	AssetLODModes.Add(MakeShared<FString>(TEXT("LOD4")));
+	AssetLODModes.Add(MakeShared<FString>(TEXT("LOD5")));
+	AssetLODModes.Add(MakeShared<FString>(TEXT("LOD6")));
+	AssetLODModes.Add(MakeShared<FString>(TEXT("LOD7")));
+	AssetLODMode = SNew(STextComboBox)
+		.OptionsSource(&AssetLODModes)
+		.OnSelectionChanged_Lambda([&](TSharedPtr<FString> String, ESelectInfo::Type)
+	{
+		EStaticMeshEditingLOD NewSelectedLOD = EStaticMeshEditingLOD::LOD0;
+		if (*String == *AssetLODModes[0])
+		{
+			NewSelectedLOD = EStaticMeshEditingLOD::MaxQuality;
+		}
+		else if (*String == *AssetLODModes[1])
+		{
+			NewSelectedLOD = EStaticMeshEditingLOD::HiResSource;
+		}
+		else
+		{
+			for (int32 k = 2; k < AssetLODModes.Num(); ++k)
+			{
+				if (*String == *AssetLODModes[k])
+				{
+					NewSelectedLOD = (EStaticMeshEditingLOD)(k - 2);
+					break;
+				}
+			}
+		}
+
+		if (FEditorInteractiveToolsFrameworkGlobals::RegisteredStaticMeshTargetFactoryKey >= 0)
+		{
+			FComponentTargetFactory* Factory = FindComponentTargetFactoryByKey(FEditorInteractiveToolsFrameworkGlobals::RegisteredStaticMeshTargetFactoryKey);
+			if (Factory != nullptr)
+			{
+				FStaticMeshComponentTargetFactory* StaticMeshFactory = static_cast<FStaticMeshComponentTargetFactory*>(Factory);
+				StaticMeshFactory->CurrentEditingLOD = NewSelectedLOD;
+			}
+		}
+
+		TObjectPtr<UToolTargetManager> TargetManager = GetScriptableEditorMode()->GetInteractiveToolsContext()->TargetManager;
+		UStaticMeshComponentToolTargetFactory* StaticMeshTargetFactory = TargetManager->FindFirstFactoryByType<UStaticMeshComponentToolTargetFactory>();
+		if (StaticMeshTargetFactory)
+		{
+			StaticMeshTargetFactory->SetActiveEditingLOD(NewSelectedLOD);
+		}
+
+	});
+
+	AssetLODModeLabel = SNew(STextBlock).Text(LOCTEXT("ActiveLODLabel", "Editing LOD"));
+
+	TSharedPtr<SVerticalBox> Content = SNew(SVerticalBox)
+	+ SVerticalBox::Slot().MaxHeight(20).HAlign(HAlign_Fill)
+		[
+			SNew(SHorizontalBox)
+			+ SHorizontalBox::Slot().Padding(0, 2, 2, 2).HAlign(HAlign_Right).VAlign(VAlign_Center).AutoWidth()
+				[ AssetLODModeLabel->AsShared() ]
+			+ SHorizontalBox::Slot().Padding(0).FillWidth(4.0f)
+				[ AssetLODMode->AsShared() ]
+		]
+
+	+ SVerticalBox::Slot().MaxHeight(20).HAlign(HAlign_Fill).Padding(0, 3, 0, 0)
+		[
+			SNew(SHorizontalBox)
+			+ SHorizontalBox::Slot().Padding(0, 2, 2, 2).HAlign(HAlign_Right).VAlign(VAlign_Center).AutoWidth()
+				[
+					SNew(STextBlock)
+					.Text(LOCTEXT("AssetLocationLabel", "New Asset Location"))
+				]
+			+ SHorizontalBox::Slot().Padding(0).FillWidth(4.0f)
+				[
+					AssetLocationMode->AsShared()
+				]
+
+			+ SHorizontalBox::Slot().HAlign(HAlign_Right).Padding(5, 0, 0, 0).FillWidth(1.0f)
+				[
+					SNew(SBox).MaxDesiredHeight(16).MaxDesiredWidth(17)
+					[
+						SNew(SButton)
+						.ContentPadding(0)
+						.ButtonStyle(FEditorStyle::Get(), "FlatButton.Default")
+						.OnClicked_Lambda( [this]() {OnShowAssetSettings(); return FReply::Handled(); } )
+						[
+							SNew(SImage).Image(FSlateIcon(FEditorStyle::GetStyleSetName(), "FoliageEditMode.Settings").GetIcon())
+						]
+					]
+				]
+		];
+
+
+	TSharedPtr<SExpandableArea> AssetConfigPanel = SNew(SExpandableArea)
+		.HeaderPadding(FMargin(2.0f))
+		.Padding(FMargin(2.f))
+		.BorderImage(FEditorStyle::Get().GetBrush("DetailsView.CategoryTop"))
+		.BorderBackgroundColor(FLinearColor(.6, .6, .6, 1.0f))
+		.BodyBorderBackgroundColor(FLinearColor::Transparent)
+		.AreaTitleFont(FEditorStyle::Get().GetFontStyle("EditorModesPanel.CategoryFontStyle"))
+		.BodyContent()
+		[
+			Content->AsShared()
+		]
+		.HeaderContent()
+		[
+			SNew(STextBlock)
+			.Text(LOCTEXT("ModelingSettingsPanelHeader", "Modeling Mode Quick Settings"))
+			.Justification(ETextJustify::Center)
+			.Font(FEditorStyle::Get().GetFontStyle("EditorModesPanel.CategoryFontStyle"))
+		];
+
+	return AssetConfigPanel;
+
+}
+
+
+void FModelingToolsEditorModeToolkit::InitializeAfterModeSetup()
+{
+	if (bFirstInitializeAfterModeSetup)
+	{
+		// force update of the active asset LOD mode, this is necessary because the update modifies
+		// ToolTarget Factories that are only available once ModelingToolsEditorMode has been initialized
+		AssetLODMode->SetSelectedItem(AssetLODModes[0]);
+
+		bFirstInitializeAfterModeSetup = false;
+	}
+
+}
+
+
+
+void FModelingToolsEditorModeToolkit::UpdateActiveToolProperties()
+{
+	UInteractiveTool* CurTool = GetScriptableEditorMode()->GetToolManager()->GetActiveTool(EToolSide::Left);
+	if (CurTool != nullptr)
+	{
+		ModeDetailsView->SetObjects(CurTool->GetToolProperties(true));
+	}
+}
+
+void FModelingToolsEditorModeToolkit::InvalidateCachedDetailPanelState(UObject* ChangedObject)
+{
+	ModeDetailsView->InvalidateCachedState();
+}
+
+
+void FModelingToolsEditorModeToolkit::PostNotification(const FText& Message)
+{
+	ClearNotification();
+
+	ActiveToolMessage = Message;
+
+	const FName LevelEditorStatusBarName = "LevelEditor.StatusBar";
+	ActiveToolMessageHandle = GEditor->GetEditorSubsystem<UStatusBarSubsystem>()->PushStatusBarMessage(LevelEditorStatusBarName, ActiveToolMessage);
+}
+
+void FModelingToolsEditorModeToolkit::ClearNotification()
+{
+	ActiveToolMessage = FText::GetEmpty();
+
+	const FName LevelEditorStatusBarName = "LevelEditor.StatusBar";
+	GEditor->GetEditorSubsystem<UStatusBarSubsystem>()->PopStatusBarMessage(LevelEditorStatusBarName, ActiveToolMessageHandle);
+	ActiveToolMessageHandle.Reset();
+}
+
+
+void FModelingToolsEditorModeToolkit::PostWarning(const FText& Message)
+{
+	ToolWarningArea->SetText(Message);
+	ToolWarningArea->SetVisibility(EVisibility::Visible);
+}
+
+void FModelingToolsEditorModeToolkit::ClearWarning()
+{
+	ToolWarningArea->SetText(FText());
+	ToolWarningArea->SetVisibility(EVisibility::Collapsed);
+}
+
+
+
+FName FModelingToolsEditorModeToolkit::GetToolkitFName() const
+{
+	return FName("ModelingToolsEditorMode");
+}
+
+FText FModelingToolsEditorModeToolkit::GetBaseToolkitName() const
+{
+	return NSLOCTEXT("ModelingToolsEditorModeToolkit", "DisplayName", "ModelingToolsEditorMode Tool");
+}
+
+static const FName PrimitiveTabName(TEXT("Shapes"));
+static const FName CreateTabName(TEXT("Create"));
+static const FName AttributesTabName(TEXT("Attributes"));
+static const FName TriModelingTabName(TEXT("TriModel"));
+static const FName PolyModelingTabName(TEXT("PolyModel"));
+static const FName MeshProcessingTabName(TEXT("MeshOps"));
+static const FName UVTabName(TEXT("UVs"));
+static const FName TransformTabName(TEXT("Transform"));
+static const FName DeformTabName(TEXT("Deform"));
+static const FName VolumesTabName(TEXT("Volumes"));
+static const FName PrototypesTabName(TEXT("Prototypes"));
+static const FName HairTabName(TEXT("Hair"));
+static const FName PolyEditTabName(TEXT("PolyEdit"));
+static const FName VoxToolsTabName(TEXT("VoxOps"));
+static const FName LODToolsTabName(TEXT("LODs"));
+
+
+const TArray<FName> FModelingToolsEditorModeToolkit::PaletteNames_Standard = { PrimitiveTabName, CreateTabName, PolyModelingTabName, TriModelingTabName, DeformTabName, TransformTabName, MeshProcessingTabName, VoxToolsTabName, AttributesTabName, UVTabName, VolumesTabName, LODToolsTabName, HairTabName };
+
+
+void FModelingToolsEditorModeToolkit::GetToolPaletteNames(TArray<FName>& PaletteNames) const
+{
+	PaletteNames = PaletteNames_Standard;
+
+	TArray<FName> ExistingNames;
+	for ( FName Name : PaletteNames )
+	{
+		ExistingNames.Add(Name);
+	}
+
+
+	bool bEnablePrototypes = (CVarEnablePrototypeModelingTools.GetValueOnGameThread() > 0);
+	if (bEnablePrototypes)
+	{
+		PaletteNames.Add(PrototypesTabName);
+		ExistingNames.Add(PrototypesTabName);
+	}
+
+	bool bEnablePolyModel = (CVarEnablePolyModeling.GetValueOnGameThread() > 0);
+	if (bEnablePolyModel)
+	{
+		PaletteNames.Add(PolyEditTabName);
+		ExistingNames.Add(PolyEditTabName);
+	}
+
+	if (IModularFeatures::Get().IsModularFeatureAvailable(IModelingModeToolExtension::GetModularFeatureName()))
+	{
+		TArray<IModelingModeToolExtension*> Extensions = IModularFeatures::Get().GetModularFeatureImplementations<IModelingModeToolExtension>(
+			IModelingModeToolExtension::GetModularFeatureName());
+		for (int32 k = 0; k < Extensions.Num(); ++k)
+		{
+			FText ExtensionName = Extensions[k]->GetExtensionName();
+			FText SectionName = Extensions[k]->GetToolSectionName();
+			FName SectionIndex(SectionName.ToString());
+			if (ExistingNames.Contains(SectionIndex))
+			{
+				UE_LOG(LogTemp, Warning, TEXT("Modeling Mode Extension [%s] uses existing Section Name [%s] - buttons may not be visible"), *ExtensionName.ToString(), *SectionName.ToString());
+			}
+			else
+			{
+				PaletteNames.Add(SectionIndex);
+				ExistingNames.Add(SectionIndex);
+			}
+		}
+	}
+}
+
+
+FText FModelingToolsEditorModeToolkit::GetToolPaletteDisplayName(FName Palette) const
+{ 
+	return FText::FromName(Palette);
+}
+
+void FModelingToolsEditorModeToolkit::BuildToolPalette(FName PaletteIndex, class FToolBarBuilder& ToolbarBuilder) 
+{
+	const FModelingToolsManagerCommands& Commands = FModelingToolsManagerCommands::Get();
+	BuildToolPalette_Experimental(PaletteIndex, ToolbarBuilder);
+}
+
+
+
+void FModelingToolsEditorModeToolkit::BuildToolPalette_Standard(FName PaletteIndex, class FToolBarBuilder& ToolbarBuilder)
+{
+	check(false);
+}
+
+
+void FModelingToolsEditorModeToolkit::BuildToolPalette_Experimental(FName PaletteIndex, class FToolBarBuilder& ToolbarBuilder)
+{
+	const FModelingToolsManagerCommands& Commands = FModelingToolsManagerCommands::Get();
+
+	if (PaletteIndex == PrimitiveTabName)
+	{
+		ToolbarBuilder.AddToolBarButton(Commands.BeginAddBoxPrimitiveTool);
+		ToolbarBuilder.AddToolBarButton(Commands.BeginAddSpherePrimitiveTool);
+		ToolbarBuilder.AddToolBarButton(Commands.BeginAddCylinderPrimitiveTool);
+		ToolbarBuilder.AddToolBarButton(Commands.BeginAddConePrimitiveTool);
+		ToolbarBuilder.AddToolBarButton(Commands.BeginAddTorusPrimitiveTool);
+		ToolbarBuilder.AddToolBarButton(Commands.BeginAddArrowPrimitiveTool);
+		ToolbarBuilder.AddToolBarButton(Commands.BeginAddRectanglePrimitiveTool);
+		ToolbarBuilder.AddToolBarButton(Commands.BeginAddDiscPrimitiveTool);
+		ToolbarBuilder.AddToolBarButton(Commands.BeginAddStairsPrimitiveTool);
+	}
+	else if (PaletteIndex == CreateTabName)
+	{
+		ToolbarBuilder.AddToolBarButton(Commands.BeginDrawPolygonTool);
+		ToolbarBuilder.AddToolBarButton(Commands.BeginDrawPolyPathTool);
+		ToolbarBuilder.AddToolBarButton(Commands.BeginDrawAndRevolveTool);
+		ToolbarBuilder.AddToolBarButton(Commands.BeginRevolveBoundaryTool);
+		ToolbarBuilder.AddToolBarButton(Commands.BeginCombineMeshesTool);
+		ToolbarBuilder.AddToolBarButton(Commands.BeginDuplicateMeshesTool);
+	}
+	else if (PaletteIndex == TransformTabName)
+	{
+		ToolbarBuilder.AddToolBarButton(Commands.BeginTransformMeshesTool);
+		ToolbarBuilder.AddToolBarButton(Commands.BeginAlignObjectsTool);
+		ToolbarBuilder.AddToolBarButton(Commands.BeginEditPivotTool);
+		ToolbarBuilder.AddToolBarButton(Commands.BeginBakeTransformTool);
+		ToolbarBuilder.AddToolBarButton(Commands.BeginTransferMeshTool);
+		ToolbarBuilder.AddToolBarButton(Commands.BeginConvertMeshesTool);
+		ToolbarBuilder.AddToolBarButton(Commands.BeginSplitMeshesTool);
+	}
+	else if (PaletteIndex == DeformTabName)
+	{
+		ToolbarBuilder.AddToolBarButton(Commands.BeginSculptMeshTool);
+		ToolbarBuilder.AddToolBarButton(Commands.BeginRemeshSculptMeshTool);
+		ToolbarBuilder.AddToolBarButton(Commands.BeginSmoothMeshTool);
+		ToolbarBuilder.AddToolBarButton(Commands.BeginOffsetMeshTool);
+		ToolbarBuilder.AddToolBarButton(Commands.BeginMeshSpaceDeformerTool);
+		ToolbarBuilder.AddToolBarButton(Commands.BeginLatticeDeformerTool);
+		ToolbarBuilder.AddToolBarButton(Commands.BeginDisplaceMeshTool);
+	}
+	else if (PaletteIndex == MeshProcessingTabName)
+	{
+		ToolbarBuilder.AddToolBarButton(Commands.BeginSimplifyMeshTool);
+		ToolbarBuilder.AddToolBarButton(Commands.BeginRemeshMeshTool);
+		ToolbarBuilder.AddToolBarButton(Commands.BeginWeldEdgesTool);
+		ToolbarBuilder.AddToolBarButton(Commands.BeginRemoveOccludedTrianglesTool);
+		ToolbarBuilder.AddToolBarButton(Commands.BeginSelfUnionTool);
+		ToolbarBuilder.AddToolBarButton(Commands.BeginProjectToTargetTool);
+	}
+	else if (PaletteIndex == LODToolsTabName)
+	{
+		ToolbarBuilder.AddToolBarButton(Commands.BeginLODManagerTool);
+		ToolbarBuilder.AddToolBarButton(Commands.BeginGenerateStaticMeshLODAssetTool);
+	}
+	else if (PaletteIndex == VoxToolsTabName)
+	{
+		ToolbarBuilder.AddToolBarButton(Commands.BeginVoxelSolidifyTool);
+		ToolbarBuilder.AddToolBarButton(Commands.BeginVoxelBlendTool);
+		ToolbarBuilder.AddToolBarButton(Commands.BeginVoxelMorphologyTool);
+		ToolbarBuilder.AddToolBarButton(Commands.BeginVoxelBooleanTool);
+		ToolbarBuilder.AddToolBarButton(Commands.BeginVoxelMergeTool);
+	}
+	else if (PaletteIndex == TriModelingTabName)
+	{
+		ToolbarBuilder.AddToolBarButton(Commands.BeginMeshSelectionTool);
+		ToolbarBuilder.AddToolBarButton(Commands.BeginTriEditTool);
+		ToolbarBuilder.AddToolBarButton(Commands.BeginHoleFillTool);
+		ToolbarBuilder.AddToolBarButton(Commands.BeginPlaneCutTool);
+		ToolbarBuilder.AddToolBarButton(Commands.BeginMirrorTool);
+		ToolbarBuilder.AddToolBarButton(Commands.BeginPolygonCutTool);
+		ToolbarBuilder.AddToolBarButton(Commands.BeginMeshTrimTool);
+	}
+	else if (PaletteIndex == PolyModelingTabName)
+	{
+		ToolbarBuilder.AddToolBarButton(Commands.BeginPolyEditTool);
+		ToolbarBuilder.AddToolBarButton(Commands.BeginPolyDeformTool);
+		ToolbarBuilder.AddToolBarButton(Commands.BeginMeshBooleanTool);
+		ToolbarBuilder.AddToolBarButton(Commands.BeginCutMeshWithMeshTool);
+		ToolbarBuilder.AddToolBarButton(Commands.BeginSubdividePolyTool);
+	}
+	else if (PaletteIndex == AttributesTabName)
+	{
+		ToolbarBuilder.AddToolBarButton(Commands.BeginMeshInspectorTool);
+		ToolbarBuilder.AddToolBarButton(Commands.BeginEditNormalsTool);
+		ToolbarBuilder.AddToolBarButton(Commands.BeginEditTangentsTool);
+		ToolbarBuilder.AddToolBarButton(Commands.BeginAttributeEditorTool);
+		ToolbarBuilder.AddToolBarButton(Commands.BeginPolyGroupsTool);
+		ToolbarBuilder.AddToolBarButton(Commands.BeginMeshGroupPaintTool);
+		ToolbarBuilder.AddToolBarButton(Commands.BeginMeshAttributePaintTool);
+		ToolbarBuilder.AddToolBarButton(Commands.BeginEditMeshMaterialsTool);
+		ToolbarBuilder.AddToolBarButton(Commands.BeginBakeMeshAttributeMapsTool);
+		ToolbarBuilder.AddToolBarButton(Commands.BeginBakeMeshAttributeVertexTool);
+	}
+	else if (PaletteIndex == UVTabName)
+	{
+		ToolbarBuilder.AddToolBarButton(Commands.BeginGlobalUVGenerateTool);
+		ToolbarBuilder.AddToolBarButton(Commands.BeginGroupUVGenerateTool);
+		ToolbarBuilder.AddToolBarButton(Commands.BeginUVProjectionTool);
+		ToolbarBuilder.AddToolBarButton(Commands.BeginUVSeamEditTool);
+		ToolbarBuilder.AddToolBarButton(Commands.BeginTransformUVIslandsTool);
+		ToolbarBuilder.AddToolBarButton(Commands.BeginUVLayoutTool);
+	}
+	else if (PaletteIndex == VolumesTabName)
+	{
+		ToolbarBuilder.AddToolBarButton(Commands.BeginVolumeToMeshTool);
+		ToolbarBuilder.AddToolBarButton(Commands.BeginMeshToVolumeTool);
+		ToolbarBuilder.AddSeparator();
+		ToolbarBuilder.AddToolBarButton(Commands.BeginBspConversionTool);
+		ToolbarBuilder.AddSeparator();
+		ToolbarBuilder.AddToolBarButton(Commands.BeginPhysicsInspectorTool);
+		ToolbarBuilder.AddToolBarButton(Commands.BeginSetCollisionGeometryTool);
+		//ToolbarBuilder.AddToolBarButton(Commands.BeginEditCollisionGeometryTool);
+		ToolbarBuilder.AddToolBarButton(Commands.BeginExtractCollisionGeometryTool);
+	}
+	else if (PaletteIndex == HairTabName)
+	{
+		ToolbarBuilder.AddToolBarButton(Commands.BeginGroomToMeshTool);
+		ToolbarBuilder.AddToolBarButton(Commands.BeginGroomCardsEditorTool);
+		ToolbarBuilder.AddToolBarButton(Commands.BeginGenerateLODMeshesTool);
+	}
+	else if (PaletteIndex == PrototypesTabName)
+	{
+		ToolbarBuilder.AddToolBarButton(Commands.BeginAddPatchTool);
+		ToolbarBuilder.AddToolBarButton(Commands.BeginShapeSprayTool);
+	}
+	else if (PaletteIndex == PolyEditTabName)
+	{
+		ToolbarBuilder.AddToolBarButton(Commands.BeginPolyModelTool_FaceSelect);
+		ToolbarBuilder.AddToolBarButton(Commands.BeginPolyModelTool_EdgeSelect);
+		ToolbarBuilder.AddToolBarButton(Commands.BeginPolyModelTool_VertexSelect);
+		ToolbarBuilder.AddToolBarButton(Commands.BeginPolyModelTool_AllSelect);
+		ToolbarBuilder.AddToolBarButton(Commands.BeginPolyModelTool_LoopSelect);
+		ToolbarBuilder.AddToolBarButton(Commands.BeginPolyModelTool_RingSelect);
+		ToolbarBuilder.AddSeparator();
+		ToolbarBuilder.AddToolBarButton(Commands.BeginPolyModelTool_Extrude);
+		ToolbarBuilder.AddToolBarButton(Commands.BeginPolyModelTool_InsetOutset);
+		ToolbarBuilder.AddToolBarButton(Commands.BeginPolyModelTool_CutFaces);
+		ToolbarBuilder.AddSeparator();
+		ToolbarBuilder.AddToolBarButton(Commands.BeginSubdividePolyTool);
+		ToolbarBuilder.AddToolBarButton(Commands.BeginPolyEditTool);
+	}
+	else
+	{
+		TArray<IModelingModeToolExtension*> Extensions = IModularFeatures::Get().GetModularFeatureImplementations<IModelingModeToolExtension>(
+			IModelingModeToolExtension::GetModularFeatureName());
+		for (int32 k = 0; k < Extensions.Num(); ++k)
+		{
+			FText SectionName = Extensions[k]->GetToolSectionName();
+			FName SectionIndex(SectionName.ToString());
+			if (PaletteIndex == SectionIndex)
+			{
+				FExtensionToolQueryInfo ExtensionQueryInfo;
+				ExtensionQueryInfo.bIsInfoQueryOnly = true;
+				TArray<FExtensionToolDescription> ToolSet;
+				Extensions[k]->GetExtensionTools(ExtensionQueryInfo, ToolSet);
+				for (const FExtensionToolDescription& ToolInfo : ToolSet)
+				{
+					ToolbarBuilder.AddToolBarButton(ToolInfo.ToolCommand);
+				}
+			}
+		}
+	}
+}
+
+void FModelingToolsEditorModeToolkit::OnToolPaletteChanged(FName PaletteName) 
+{
+}
+
+
+
+void FModelingToolsEditorModeToolkit::EnableShowRealtimeWarning(bool bEnable)
+{
+	if (bShowRealtimeWarning != bEnable)
+	{
+		bShowRealtimeWarning = bEnable;
+		UpdateShowWarnings();
+	}
+}
+
+void FModelingToolsEditorModeToolkit::OnToolStarted(UInteractiveToolManager* Manager, UInteractiveTool* Tool)
+{
+	UpdateActiveToolProperties();
+
+	UInteractiveTool* CurTool = GetScriptableEditorMode()->GetToolManager()->GetActiveTool(EToolSide::Left);
+	CurTool->OnPropertySetsModified.AddSP(this, &FModelingToolsEditorModeToolkit::UpdateActiveToolProperties);
+	CurTool->OnPropertyModifiedDirectlyByTool.AddSP(this, &FModelingToolsEditorModeToolkit::InvalidateCachedDetailPanelState);
+
+	ModeHeaderArea->SetVisibility(EVisibility::Collapsed);
+	ActiveToolName = CurTool->GetToolInfo().ToolDisplayName;
+
+	// try to update icon
+	FString ActiveToolIdentifier = GetScriptableEditorMode()->GetToolManager()->GetActiveToolName(EToolSide::Left);
+	ActiveToolIdentifier.InsertAt(0, ".");
+	FName ActiveToolIconName = ISlateStyle::Join(FModelingToolsManagerCommands::Get().GetContextName(), TCHAR_TO_ANSI(*ActiveToolIdentifier));
+	ActiveToolIcon = FModelingToolsEditorModeStyle::Get()->GetOptionalBrush(ActiveToolIconName);
+
+
+	GetToolkitHost()->AddViewportOverlayWidget(ViewportOverlayWidget.ToSharedRef());
+
+	// disable LOD level picker once Tool is active
+	AssetLODMode->SetEnabled(false);
+	AssetLODModeLabel->SetEnabled(false);
+}
+
+void FModelingToolsEditorModeToolkit::OnToolEnded(UInteractiveToolManager* Manager, UInteractiveTool* Tool)
+{
+	if (IsHosted())
+	{
+		GetToolkitHost()->RemoveViewportOverlayWidget(ViewportOverlayWidget.ToSharedRef());
+	}
+
+	ModeDetailsView->SetObject(nullptr);
+	ActiveToolName = FText::GetEmpty();
+	ModeHeaderArea->SetVisibility(EVisibility::Visible);
+	ModeHeaderArea->SetText(LOCTEXT("SelectToolLabel", "Select a Tool from the Toolbar"));
+	ClearNotification();
+	ClearWarning();
+	UInteractiveTool* CurTool = GetScriptableEditorMode()->GetToolManager()->GetActiveTool(EToolSide::Left);
+	if ( CurTool )
+	{
+		CurTool->OnPropertySetsModified.RemoveAll(this);
+		CurTool->OnPropertyModifiedDirectlyByTool.RemoveAll(this);
+	}
+
+	// re-enable LOD level picker
+	AssetLODMode->SetEnabled(true);
+	AssetLODModeLabel->SetEnabled(true);
+}
+
+void FModelingToolsEditorModeToolkit::OnActiveViewportChanged(TSharedPtr<IAssetViewport> OldViewport, TSharedPtr<IAssetViewport> NewViewport)
+{
+	// Only worry about handling this notification if Modeling has an active tool
+	if (!ActiveToolName.IsEmpty())
+	{
+		// Check first to see if this changed because the old viewport was deleted and if not, remove our hud
+		if (OldViewport)	
+		{
+			GetToolkitHost()->RemoveViewportOverlayWidget(ViewportOverlayWidget.ToSharedRef(), OldViewport);
+		}
+
+		// Add the hud to the new viewport
+		GetToolkitHost()->AddViewportOverlayWidget(ViewportOverlayWidget.ToSharedRef(), NewViewport);
+	}
+}
+
+void FModelingToolsEditorModeToolkit::UpdateShowWarnings()
+{
+	if (bShowRealtimeWarning )
+	{
+		if (ModeWarningArea->GetVisibility() == EVisibility::Collapsed)
+		{
+			ModeWarningArea->SetText(LOCTEXT("ModelingModeToolkitRealtimeWarning", "Realtime Mode is required for Modeling Tools to work correctly. Please enable Realtime Mode in the Viewport Options or with the Ctrl+r hotkey."));
+			ModeWarningArea->SetVisibility(EVisibility::Visible);
+		}
+	}
+	else
+	{
+		ModeWarningArea->SetText(FText());
+		ModeWarningArea->SetVisibility(EVisibility::Collapsed);
+	}
+
+}
+
+
+void FModelingToolsEditorModeToolkit::UpdateAssetLocationMode(TSharedPtr<FString> NewString)
+{
+	UModelingToolsEditorModeSettings* Settings = GetMutableDefault<UModelingToolsEditorModeSettings>();
+	if (NewString == AssetLocationModes[0])
+	{
+		Settings->AssetGenerationLocation = EModelingModeAssetGenerationLocation::AutoGeneratedWorldRelativeAssetPath;
+	}
+	if (NewString == AssetLocationModes[1])
+	{
+		Settings->AssetGenerationLocation = EModelingModeAssetGenerationLocation::AutoGeneratedGlobalAssetPath;
+	}
+	else if ( NewString == AssetLocationModes[2])
+	{
+		Settings->AssetGenerationLocation = EModelingModeAssetGenerationLocation::CurrentAssetBrowserPathIfAvailable;
+	}
+	else
+	{
+		Settings->AssetGenerationLocation = EModelingModeAssetGenerationLocation::AutoGeneratedWorldRelativeAssetPath;
+	}
+
+	Settings->SaveConfig();
+}
+
+void FModelingToolsEditorModeToolkit::UpdateAssetSaveMode(TSharedPtr<FString> NewString)
+{
+	UModelingToolsEditorModeSettings* Settings = GetMutableDefault<UModelingToolsEditorModeSettings>();
+	if (NewString == AssetSaveModes[0])
+	{
+		Settings->AssetGenerationMode = EModelingModeAssetGenerationBehavior::AutoGenerateAndAutosave;
+	}
+	else if (NewString == AssetSaveModes[1])
+	{
+		Settings->AssetGenerationMode = EModelingModeAssetGenerationBehavior::AutoGenerateButDoNotAutosave;
+	}
+	else if (NewString == AssetSaveModes[2])
+	{
+		Settings->AssetGenerationMode = EModelingModeAssetGenerationBehavior::InteractivePromptToSave;
+	}
+	else
+	{
+		Settings->AssetGenerationMode = EModelingModeAssetGenerationBehavior::AutoGenerateButDoNotAutosave;
+	}
+
+	Settings->SaveConfig();
+}
+
+void FModelingToolsEditorModeToolkit::UpdateAssetPanelFromSettings()
+{
+	const UModelingToolsEditorModeSettings* Settings = GetDefault<UModelingToolsEditorModeSettings>();
+
+	switch (Settings->AssetGenerationLocation)
+	{
+	case EModelingModeAssetGenerationLocation::CurrentAssetBrowserPathIfAvailable:
+		AssetLocationMode->SetSelectedItem(AssetLocationModes[2]);
+		break;
+	case EModelingModeAssetGenerationLocation::AutoGeneratedGlobalAssetPath:
+		AssetLocationMode->SetSelectedItem(AssetLocationModes[1]);
+		break;
+	case EModelingModeAssetGenerationLocation::AutoGeneratedWorldRelativeAssetPath:
+	default:
+		AssetLocationMode->SetSelectedItem(AssetLocationModes[0]);
+		break;
+	}
+
+	switch (Settings->AssetGenerationMode)
+	{
+	case EModelingModeAssetGenerationBehavior::AutoGenerateButDoNotAutosave:
+		AssetSaveMode->SetSelectedItem(AssetSaveModes[1]);
+		break;
+	case EModelingModeAssetGenerationBehavior::InteractivePromptToSave:
+		AssetSaveMode->SetSelectedItem(AssetSaveModes[2]);
+		break;
+	default:
+		AssetSaveMode->SetSelectedItem(AssetSaveModes[0]);
+		break;
+	}
+}
+
+
+void FModelingToolsEditorModeToolkit::UpdateObjectCreationOptionsFromSettings()
+{
+	// update DynamicMeshActor Settings
+	const UModelingToolsEditorModeSettings* Settings = GetDefault<UModelingToolsEditorModeSettings>();
+
+	// enable/disable dynamic mesh actors
+	UCreateMeshObjectTypeProperties::bEnableDynamicMeshActorSupport = Settings->bEnableDynamicMeshActors;
+
+	// set configured default type
+	if (Settings->DefaultMeshObjectType == EModelingModeDefaultMeshObjectType::DynamicMeshActor && Settings->bEnableDynamicMeshActors)
+	{
+		UCreateMeshObjectTypeProperties::DefaultObjectTypeIdentifier = UCreateMeshObjectTypeProperties::DynamicMeshActorIdentifier;
+	}
+	else if (Settings->DefaultMeshObjectType == EModelingModeDefaultMeshObjectType::VolumeActor)
+	{
+		UCreateMeshObjectTypeProperties::DefaultObjectTypeIdentifier = UCreateMeshObjectTypeProperties::VolumeIdentifier;
+	}
+	else
+	{
+		UCreateMeshObjectTypeProperties::DefaultObjectTypeIdentifier = UCreateMeshObjectTypeProperties::StaticMeshIdentifier;
+	}
+}
+
+void FModelingToolsEditorModeToolkit::OnAssetSettingsModified()
+{
+	UpdateObjectCreationOptionsFromSettings();
+	UpdateAssetPanelFromSettings();
+}
+
+void FModelingToolsEditorModeToolkit::OnShowAssetSettings()
+{
+	if (ISettingsModule* SettingsModule = FModuleManager::GetModulePtr<ISettingsModule>("Settings"))
+	{
+		SettingsModule->ShowViewer("Project", "Plugins", "ModelingMode");
+	}
+}
+
+
+#undef LOCTEXT_NAMESPACE
diff --git a/Engine/Plugins/Editor/ModelingToolsEditorMode/Source/ModelingToolsEditorMode/Private/ModelingToolsManagerActions.cpp b/Engine/Plugins/Editor/ModelingToolsEditorMode/Source/ModelingToolsEditorMode/Private/ModelingToolsManagerActions.cpp
new file mode 100644
--- /dev/null
+++ b/Engine/Plugins/Editor/ModelingToolsEditorMode/Source/ModelingToolsEditorMode/Private/ModelingToolsManagerActions.cpp
@@ -0,0 +1,145 @@
+// Copyright Epic Games, Inc. All Rights Reserved.
+
+#include "ModelingToolsManagerActions.h"
+#include "EditorStyleSet.h"
+#include "InputCoreTypes.h"
+#include "ModelingToolsEditorModeStyle.h"
+
+#define LOCTEXT_NAMESPACE "ModelingToolsManagerCommands"
+
+
+
+FModelingToolsManagerCommands::FModelingToolsManagerCommands() :
+	TCommands<FModelingToolsManagerCommands>(
+		"ModelingToolsManagerCommands", // Context name for fast lookup
+		NSLOCTEXT("Contexts", "ModelingToolsToolCommands", "Modeling Mode - Tools"), // Localized context name for displaying
+		NAME_None, // Parent
+		FModelingToolsEditorModeStyle::Get()->GetStyleSetName() // Icon Style Set
+		)
+{
+}
+
+
+void FModelingToolsManagerCommands::RegisterCommands()
+{
+	UI_COMMAND(BeginAddBoxPrimitiveTool, "Box", "Create new Box StaticMesh Assets", EUserInterfaceActionType::ToggleButton, FInputChord());
+	UI_COMMAND(BeginAddCylinderPrimitiveTool, "Cyl", "Create new Cylinder StaticMesh Assets", EUserInterfaceActionType::ToggleButton, FInputChord());
+	UI_COMMAND(BeginAddConePrimitiveTool, "Cone", "Create new Cone StaticMesh Assets", EUserInterfaceActionType::ToggleButton, FInputChord());
+	UI_COMMAND(BeginAddArrowPrimitiveTool, "Arrow", "Create new Arrow StaticMesh Assets", EUserInterfaceActionType::ToggleButton, FInputChord());
+	UI_COMMAND(BeginAddRectanglePrimitiveTool, "Rect", "Create new Rectangle StaticMesh Assets", EUserInterfaceActionType::ToggleButton, FInputChord());
+	UI_COMMAND(BeginAddDiscPrimitiveTool, "Disc", "Create new Disc StaticMesh Assets", EUserInterfaceActionType::ToggleButton, FInputChord());
+	UI_COMMAND(BeginAddTorusPrimitiveTool, "Torus", "Create new Torus StaticMesh Assets", EUserInterfaceActionType::ToggleButton, FInputChord());
+	UI_COMMAND(BeginAddSpherePrimitiveTool, "Sphere", "Create new Sphere StaticMesh Assets", EUserInterfaceActionType::ToggleButton, FInputChord());
+	UI_COMMAND(BeginAddStairsPrimitiveTool, "Stair", "Create new Stairs StaticMesh Assets", EUserInterfaceActionType::ToggleButton, FInputChord());
+
+	UI_COMMAND(BeginAddPatchTool, "Patch", "Add Patch", EUserInterfaceActionType::ToggleButton, FInputChord());
+	UI_COMMAND(BeginRevolveBoundaryTool, "BdryRev", "Revolve Mesh Boundary Loop into New Asset", EUserInterfaceActionType::ToggleButton, FInputChord());
+	UI_COMMAND(BeginDrawPolygonTool, "PolyExt", "Draw/Extrude Polygons to create new StaticMesh Assets", EUserInterfaceActionType::ToggleButton, FInputChord());
+	UI_COMMAND(BeginDrawPolyPathTool, "PathExt", "Draw/Extrude PolyPaths to create new StaticMesh Assets", EUserInterfaceActionType::ToggleButton, FInputChord());
+	UI_COMMAND(BeginDrawAndRevolveTool, "PathRev", "Draw and Revolve a PolyPath to create a new StaticMesh Assets", EUserInterfaceActionType::ToggleButton, FInputChord());
+	UI_COMMAND(BeginShapeSprayTool, "Spray", "Shape Spray", EUserInterfaceActionType::ToggleButton, FInputChord());
+
+	UI_COMMAND(BeginSculptMeshTool, "VSclpt", "Vertex Sculpting", EUserInterfaceActionType::ToggleButton, FInputChord());
+	UI_COMMAND(BeginPolyEditTool, "PolyEd", "Edit Mesh via PolyGroups", EUserInterfaceActionType::ToggleButton, FInputChord());
+	UI_COMMAND(BeginTriEditTool, "TriEd", "Edit Mesh via Triangles", EUserInterfaceActionType::ToggleButton, FInputChord());
+	UI_COMMAND(BeginPolyDeformTool, "PolyDef", "Deform Mesh via PolyGroups", EUserInterfaceActionType::ToggleButton, FInputChord());
+	UI_COMMAND(BeginSmoothMeshTool, "Smooth", "Smooth Mesh surface", EUserInterfaceActionType::ToggleButton, FInputChord());
+	UI_COMMAND(BeginOffsetMeshTool, "Offset", "Offset Mesh surface", EUserInterfaceActionType::ToggleButton, FInputChord());
+	UI_COMMAND(BeginDisplaceMeshTool, "Displce", "Displace Mesh surface with optional subdivision", EUserInterfaceActionType::ToggleButton, FInputChord());
+	UI_COMMAND(BeginMeshSpaceDeformerTool, "Warp", "Reshape Mesh using Space Deformers", EUserInterfaceActionType::ToggleButton, FInputChord());
+	UI_COMMAND(BeginTransformMeshesTool, "XForm", "Transform selected Meshes", EUserInterfaceActionType::ToggleButton, FInputChord());
+	UI_COMMAND(BeginEditPivotTool, "Pivot", "Edit Mesh Pivots", EUserInterfaceActionType::ToggleButton, FInputChord());
+	UI_COMMAND(BeginBakeTransformTool, "BakeRS", "Bake Scale/Rotation into Mesh Asset", EUserInterfaceActionType::ToggleButton, FInputChord());
+	UI_COMMAND(BeginCombineMeshesTool, "Append", "Combine Selection into new StaticMesh Asset", EUserInterfaceActionType::ToggleButton, FInputChord());
+	UI_COMMAND(BeginDuplicateMeshesTool, "Dupe", "Duplicate Selection into new StaticMesh Asset", EUserInterfaceActionType::ToggleButton, FInputChord());
+	UI_COMMAND(BeginAlignObjectsTool, "Align", "Align Objects", EUserInterfaceActionType::ToggleButton, FInputChord());
+	UI_COMMAND(BeginTransferMeshTool, "Transfer", "Transfer Mesh", EUserInterfaceActionType::ToggleButton, FInputChord());
+	UI_COMMAND(BeginConvertMeshesTool, "Convert", "Convert Meshes", EUserInterfaceActionType::ToggleButton, FInputChord());
+	UI_COMMAND(BeginSplitMeshesTool, "Split", "Split Meshes", EUserInterfaceActionType::ToggleButton, FInputChord());
+
+	UI_COMMAND(BeginRemeshSculptMeshTool, "DSclpt", "Dynamic Mesh Sculpting", EUserInterfaceActionType::ToggleButton, FInputChord());
+	UI_COMMAND(BeginRemeshMeshTool, "Remesh", "Retriangulate Mesh", EUserInterfaceActionType::ToggleButton, FInputChord());
+	UI_COMMAND(BeginProjectToTargetTool, "Project", "Map/Remesh onto Target Mesh (second selection)", EUserInterfaceActionType::ToggleButton, FInputChord());
+	UI_COMMAND(BeginSimplifyMeshTool, "Simplfy", "Simplify Mesh", EUserInterfaceActionType::ToggleButton, FInputChord());
+	UI_COMMAND(BeginEditNormalsTool, "Nrmls", "Modify/Update Mesh Normals", EUserInterfaceActionType::ToggleButton, FInputChord());
+	UI_COMMAND(BeginEditTangentsTool, "Tngnts", "Update Mesh Tangents", EUserInterfaceActionType::ToggleButton, FInputChord());
+	UI_COMMAND(BeginRemoveOccludedTrianglesTool, "Jacket", "Remove Hidden Triangles from selected Meshes", EUserInterfaceActionType::ToggleButton, FInputChord());
+	UI_COMMAND(BeginUVProjectionTool, "Project", "Set UVs from Projection", EUserInterfaceActionType::ToggleButton, FInputChord());
+	UI_COMMAND(BeginUVLayoutTool, "Layout", "Re-Pack Existing UVs", EUserInterfaceActionType::ToggleButton, FInputChord());
+	UI_COMMAND(BeginVoxelMergeTool, "VoxMrg", "Merge Selected Meshes (Voxel-Based)", EUserInterfaceActionType::ToggleButton, FInputChord());
+	UI_COMMAND(BeginVoxelBooleanTool, "VoxBool", "Boolean Selected Meshes (Voxel-Based)", EUserInterfaceActionType::ToggleButton, FInputChord());
+	UI_COMMAND(BeginVoxelSolidifyTool, "VoxWrap", "Wrap Selected Meshes (Voxel-Based)", EUserInterfaceActionType::ToggleButton, FInputChord());
+	UI_COMMAND(BeginVoxelBlendTool, "VoxBlnd", "Blend Selected Meshes (Voxel-Based)", EUserInterfaceActionType::ToggleButton, FInputChord());
+	UI_COMMAND(BeginVoxelMorphologyTool, "VoxMrph", "Offset/Inset Selected Meshes (Voxel-Based)", EUserInterfaceActionType::ToggleButton, FInputChord());
+	UI_COMMAND(BeginSelfUnionTool, "Merge", "Self-Union Selected Meshes to resolve Self-Intersections", EUserInterfaceActionType::ToggleButton, FInputChord());
+	UI_COMMAND(BeginMeshBooleanTool, "Boolean", "Apply Boolean/CSG operation to selected Mesh pair", EUserInterfaceActionType::ToggleButton, FInputChord());
+	UI_COMMAND(BeginMeshTrimTool, "Trim", "Trim/Cut selected mesh with second mesh", EUserInterfaceActionType::ToggleButton, FInputChord());
+	UI_COMMAND(BeginCutMeshWithMeshTool, "MshCut", "Split selected mesh into parts using second mesh", EUserInterfaceActionType::ToggleButton, FInputChord());
+	UI_COMMAND(BeginBspConversionTool, "BSPConv", "Convert BSP to StaticMesh Asset", EUserInterfaceActionType::ToggleButton, FInputChord());
+	UI_COMMAND(BeginMeshToVolumeTool, "Msh2Vol", "Convert Mesh to Volume", EUserInterfaceActionType::ToggleButton, FInputChord());
+	UI_COMMAND(BeginVolumeToMeshTool, "Vol2Msh", "Convert Volume to new Mesh Asset", EUserInterfaceActionType::ToggleButton, FInputChord());
+	UI_COMMAND(BeginPlaneCutTool, "PlnCut", "Cut Selected Meshes with Plane", EUserInterfaceActionType::ToggleButton, FInputChord());
+	UI_COMMAND(BeginMirrorTool, "Mirror", "Mirror Selected Meshes", EUserInterfaceActionType::ToggleButton, FInputChord());
+	UI_COMMAND(BeginHoleFillTool, "HFill", "Fill Holes in Mesh", EUserInterfaceActionType::ToggleButton, FInputChord());
+	UI_COMMAND(BeginLatticeDeformerTool, "Lattice", "Deform Mesh with 3D Lattice/Grid", EUserInterfaceActionType::ToggleButton, FInputChord());
+	UI_COMMAND(BeginPolygonCutTool, "PolyCut", "Cut Mesh with Extruded Polygon", EUserInterfaceActionType::ToggleButton, FInputChord());
+	UI_COMMAND(BeginGlobalUVGenerateTool, "AutoUV", "Auto-Unwrap and Pack UVs for Mesh", EUserInterfaceActionType::ToggleButton, FInputChord());
+	UI_COMMAND(BeginGroupUVGenerateTool, "Unwrap", "Recalculate UV unwrapping for Mesh regions", EUserInterfaceActionType::ToggleButton, FInputChord());
+	UI_COMMAND(BeginMeshSelectionTool, "TriSel", "Select and Edit Mesh Triangles", EUserInterfaceActionType::ToggleButton, FInputChord());
+
+	UI_COMMAND(BeginPhysicsInspectorTool, "PInspct", "Inspect Physics Geometry for selected Meshes", EUserInterfaceActionType::ToggleButton, FInputChord());
+	UI_COMMAND(BeginSetCollisionGeometryTool, "Msh2Coll", "Convert Selected Meshes to Simple Collision Geometry (for last Selected)", EUserInterfaceActionType::ToggleButton, FInputChord());
+	UI_COMMAND(BeginEditCollisionGeometryTool, "EditPhys", "Edit Simple Collision Geometry for selected Mesh", EUserInterfaceActionType::ToggleButton, FInputChord());
+	UI_COMMAND(BeginExtractCollisionGeometryTool, "Coll2Msh", "Convert Simple Collision Geometry to new Mesh Asset", EUserInterfaceActionType::ToggleButton, FInputChord());
+
+	UI_COMMAND(BeginMeshInspectorTool, "Inspct", "Inspect Mesh Attributes", EUserInterfaceActionType::ToggleButton, FInputChord());
+	UI_COMMAND(BeginWeldEdgesTool, "Weld", "Weld Overlapping Mesh Edges", EUserInterfaceActionType::ToggleButton, FInputChord());
+	UI_COMMAND(BeginPolyGroupsTool, "GenGrps", "Generate new PolyGroups", EUserInterfaceActionType::ToggleButton, FInputChord());
+	UI_COMMAND(BeginEditMeshMaterialsTool, "MatEd", "Assign Materials to Selected Triangles", EUserInterfaceActionType::ToggleButton, FInputChord());
+	UI_COMMAND(BeginTransformUVIslandsTool, "XForm", "Transform UV Islands in UV Space", EUserInterfaceActionType::ToggleButton, FInputChord());
+	UI_COMMAND(BeginMeshAttributePaintTool, "MapPnt", "Paint Attribute Maps", EUserInterfaceActionType::ToggleButton, FInputChord());
+	UI_COMMAND(BeginAttributeEditorTool, "AttrEd", "Edit/Configure Mesh Attributes", EUserInterfaceActionType::ToggleButton, FInputChord());
+	UI_COMMAND(BeginBakeMeshAttributeMapsTool, "BakeTx", "Bake Image Maps for Target Mesh (optionally from second Source Mesh)", EUserInterfaceActionType::ToggleButton, FInputChord());
+	UI_COMMAND(BeginBakeMeshAttributeVertexTool, "BakeVtx", "Bake Vertex Colors for Target Mesh (optionally from a second Source Mesh)", EUserInterfaceActionType::ToggleButton, FInputChord());
+	UI_COMMAND(BeginUVSeamEditTool, "SeamEd", "Add UV Seams to Mesh", EUserInterfaceActionType::ToggleButton, FInputChord());
+
+	UI_COMMAND(BeginSkinWeightsPaintTool, "SkinWts", "Start the Paint Skin Weights Tool", EUserInterfaceActionType::Button, FInputChord());
+	UI_COMMAND(BeginSkinWeightsBindingTool, "SkinBind", "Start the Skin Weights Binding Tool", EUserInterfaceActionType::Button, FInputChord());
+
+	UI_COMMAND(BeginGroomToMeshTool, "Helmet", "Generate Helmet Mesh for Selected Groom", EUserInterfaceActionType::ToggleButton, FInputChord());
+	UI_COMMAND(BeginGenerateLODMeshesTool, "HlmLOD", "Generate LODS for Hair Helmet", EUserInterfaceActionType::ToggleButton, FInputChord());
+	UI_COMMAND(BeginGroomCardsToMeshTool, "CardsToMesh", "Hair Cards to Mesh Tool", EUserInterfaceActionType::ToggleButton, FInputChord());
+	UI_COMMAND(BeginGroomMeshToCardsTool, "MeshToCards", "Mesh to Hair Cards Tool", EUserInterfaceActionType::ToggleButton, FInputChord());
+	UI_COMMAND(BeginGroomCardsEditorTool, "CardsEd", "Edit Hair Cards", EUserInterfaceActionType::ToggleButton, FInputChord());
+	UI_COMMAND(BeginSubdividePolyTool, "SubD", "Subdivide Mesh Polygroups or Triangles", EUserInterfaceActionType::ToggleButton, FInputChord());
+
+	
+	UI_COMMAND(BeginLODManagerTool, "LODMgr", "Static Mesh Asset LOD Manager", EUserInterfaceActionType::ToggleButton, FInputChord());
+	UI_COMMAND(BeginGenerateStaticMeshLODAssetTool, "AutoLOD", "Generate Static Mesh LOD Asset", EUserInterfaceActionType::ToggleButton, FInputChord());
+	UI_COMMAND(BeginMeshGroupPaintTool, "GrpPnt", "Paint New Mesh Polygroups", EUserInterfaceActionType::ToggleButton, FInputChord());
+
+
+	UI_COMMAND(BeginPolyModelTool_FaceSelect, "Faces", "PolyGroup Face Selection Tool", EUserInterfaceActionType::ToggleButton, FInputChord());
+	UI_COMMAND(BeginPolyModelTool_EdgeSelect, "Edges", "PolyGroup Edge Selection Tool", EUserInterfaceActionType::ToggleButton, FInputChord());
+	UI_COMMAND(BeginPolyModelTool_VertexSelect, "Verts", "PolyGroup Vertex Selection Tool", EUserInterfaceActionType::ToggleButton, FInputChord());
+	UI_COMMAND(BeginPolyModelTool_AllSelect, "Any", "PolyGroup Face/Edge/Vertex Selection Tool", EUserInterfaceActionType::ToggleButton, FInputChord());
+	UI_COMMAND(BeginPolyModelTool_LoopSelect, "Loops", "PolyGroup Loop Selection Tool", EUserInterfaceActionType::ToggleButton, FInputChord());
+	UI_COMMAND(BeginPolyModelTool_RingSelect, "Rings", "PolyGroup Ring Selection Tool", EUserInterfaceActionType::ToggleButton, FInputChord());
+
+	UI_COMMAND(BeginPolyModelTool_Extrude, "Extrude", "PolyGroup Extrude Tool", EUserInterfaceActionType::ToggleButton, FInputChord());
+	UI_COMMAND(BeginPolyModelTool_InsetOutset, "Inset/Outset", "PolyGroup Inset/Outset Tool", EUserInterfaceActionType::ToggleButton, FInputChord());
+	UI_COMMAND(BeginPolyModelTool_CutFaces, "Cut", "PolyGroup Cut Faces Tool", EUserInterfaceActionType::ToggleButton, FInputChord());
+
+	UI_COMMAND(AcceptActiveTool, "Accept", "Accept the active tool", EUserInterfaceActionType::Button, FInputChord());
+	UI_COMMAND(CancelActiveTool, "Cancel", "Cancel the active tool", EUserInterfaceActionType::Button, FInputChord());
+	UI_COMMAND(CompleteActiveTool, "Done", "Complete the active tool", EUserInterfaceActionType::Button, FInputChord());
+
+	// Note that passing a chord into one of these calls hooks the key press to the respective action. 
+	UI_COMMAND(AcceptOrCompleteActiveTool, "Accept or Complete", "Accept or complete the active tool", EUserInterfaceActionType::Button, FInputChord(EKeys::Enter));
+	UI_COMMAND(CancelOrCompleteActiveTool, "Cancel or Complete", "Cancel or complete the active tool", EUserInterfaceActionType::Button, FInputChord(EKeys::Escape));
+}
+
+
+
+
+#undef LOCTEXT_NAMESPACE
diff --git a/Engine/Plugins/Editor/ModelingToolsEditorMode/Source/ModelingToolsEditorMode/Public/DetailsCustomizations/ModelingToolPropertyCustomizations.h b/Engine/Plugins/Editor/ModelingToolsEditorMode/Source/ModelingToolsEditorMode/Public/DetailsCustomizations/ModelingToolPropertyCustomizations.h
new file mode 100644
--- /dev/null
+++ b/Engine/Plugins/Editor/ModelingToolsEditorMode/Source/ModelingToolsEditorMode/Public/DetailsCustomizations/ModelingToolPropertyCustomizations.h
@@ -0,0 +1,20 @@
+// Copyright Epic Games, Inc. All Rights Reserved.
+
+#pragma once
+
+#include "IPropertyTypeCustomization.h"
+#include "Input/Reply.h"
+#include "Widgets/Text/STextBlock.h"
+#include "Widgets/Input/SButton.h"
+
+class IDetailLayoutBuilder;
+
+/** Details customization for FModelingModeAxisFilter struct - X/Y/Z booleans are laid out in a single horizontal row */
+class FModelingToolsAxisFilterCustomization : public IPropertyTypeCustomization
+{
+public:
+	static TSharedRef<IPropertyTypeCustomization> MakeInstance();
+
+	virtual void CustomizeHeader(TSharedRef<IPropertyHandle> StructPropertyHandle, class FDetailWidgetRow& HeaderRow, IPropertyTypeCustomizationUtils& StructCustomizationUtils) override;
+	virtual void CustomizeChildren(TSharedRef<IPropertyHandle> StructPropertyHandle, class IDetailChildrenBuilder& ChildBuilder, IPropertyTypeCustomizationUtils& StructCustomizationUtils) override;
+};
diff --git a/Engine/Plugins/Editor/ModelingToolsEditorMode/Source/ModelingToolsEditorMode/Public/ModelingModeAssetUtils.h b/Engine/Plugins/Editor/ModelingToolsEditorMode/Source/ModelingToolsEditorMode/Public/ModelingModeAssetUtils.h
new file mode 100644
--- /dev/null
+++ b/Engine/Plugins/Editor/ModelingToolsEditorMode/Source/ModelingToolsEditorMode/Public/ModelingModeAssetUtils.h
@@ -0,0 +1,45 @@
+// Copyright Epic Games, Inc. All Rights Reserved.
+
+#pragma once
+
+#include "CoreMinimal.h"
+
+class UWorld;
+
+
+namespace UE
+{
+namespace Modeling
+{
+
+//
+// Utility functions that can be used to plug in the client-configurable parts of UEditorModelingObjectsCreationAPI.
+// See ModelingToolsEditorMode for usage.
+//
+
+/**
+ * Determines path and name for a new Asset based on current mode settings, etc
+ * @param BaseName desired initial string for the name
+ * @param TargetWorld World the new Asset will be used in, which can determine its generated path based on various settings
+ * @param SuggestedFolder a suggested path for the asset, ignored if empty
+ * @return Path+Name for the new Asset (package)
+ */
+MODELINGTOOLSEDITORMODE_API FString GetNewAssetPathName(const FString& BaseName, const UWorld* TargetWorld, FString SuggestedFolder);
+
+
+/**
+ * Utility function that may auto-save a newly created asset depending on current mode settings,
+ * and posts notifications for the Editor
+ * @param Asset the new asset object (eg UStaticMesh, UTexture2D, etc)
+ */
+MODELINGTOOLSEDITORMODE_API void OnNewAssetCreated(UObject* Asset);
+
+
+/**
+ * Utility function that will attempt to auto-save an asset, and post Editor change notifications
+ * @param Asset the  asset object (eg UStaticMesh, UTexture2D, etc)
+ */
+MODELINGTOOLSEDITORMODE_API bool AutoSaveAsset(UObject* Asset);
+
+}
+}
\ No newline at end of file
diff --git a/Engine/Plugins/Editor/ModelingToolsEditorMode/Source/ModelingToolsEditorMode/Public/ModelingModeToolExtensions.h b/Engine/Plugins/Editor/ModelingToolsEditorMode/Source/ModelingToolsEditorMode/Public/ModelingModeToolExtensions.h
new file mode 100644
--- /dev/null
+++ b/Engine/Plugins/Editor/ModelingToolsEditorMode/Source/ModelingToolsEditorMode/Public/ModelingModeToolExtensions.h
@@ -0,0 +1,92 @@
+// Copyright Epic Games, Inc. All Rights Reserved.
+
+#pragma once
+
+#include "CoreMinimal.h"
+#include "Features/IModularFeature.h"
+
+class UInteractiveToolsContext;
+class FUICommandInfo;
+class UInteractiveToolBuilder;
+
+
+// dummy class for extension API below
+class FModelingModeAssetAPI
+{
+public:
+};
+
+
+/**
+ * This struct is passed to IModelingModeToolExtension implementations to allow 
+ * them to forward various Tools Context information to their ToolBuilders/etc
+ * 
+ * Note that if bIsInfoQueryOnly is true, then the Extension does not need to 
+ * construct Tool Builders, the query is for information purposes only
+ */
+struct MODELINGTOOLSEDITORMODE_API FExtensionToolQueryInfo
+{
+	bool bIsInfoQueryOnly = false;
+
+	UInteractiveToolsContext* ToolsContext = nullptr;
+	FModelingModeAssetAPI* AssetAPI = nullptr;
+};
+
+/**
+ * IModelingModeToolExtension implementations return the list of Tools they provide
+ * via instances of FExtensionToolDescription.
+ */
+struct MODELINGTOOLSEDITORMODE_API FExtensionToolDescription
+{
+	/** Long name of the Tool, used in various places in the UI */
+	FText ToolName;
+	/** Command that is added to the Tool button set. This defines the button label. */
+	TSharedPtr<FUICommandInfo> ToolCommand;
+	/** Builder for the Tool. This can be null if FExtensionToolQueryInfo.bIsInfoQueryOnly is true. */
+	UInteractiveToolBuilder* ToolBuilder = nullptr;
+};
+
+/**
+ * IModelingModeToolExtension uses the IModularFeature API to allow a Plugin to provide
+ * a set of InteractiveTool's to be exposed in Modeling Mode. The Tools will be 
+ * included in a section of the Modeling Mode tool list, based on GetToolSectionName().
+ * 
+ */
+class MODELINGTOOLSEDITORMODE_API IModelingModeToolExtension : public IModularFeature
+{
+public:
+	virtual ~IModelingModeToolExtension() {}
+
+	/**
+	 * @return A text string that identifies this Extension.
+	 */
+	virtual FText GetExtensionName() = 0;
+
+	/**
+	 * @return A text string that defines the name of the Toolbar Section this Extension's tools will be included in.
+	 * @warning If the same Section is used in multiple Extensions, some buttons may not be shown
+	 */
+	virtual FText GetToolSectionName() = 0;
+
+	/**
+	 * Query the Extension for a list of Tools to expose in Modeling Mode.
+	 * Note that this function *will* be called multiple times by Modeling Mode, as the
+	 * information about the set of Tools is needed in multiple places. The QueryInfo.bIsInfoQueryOnly
+	 * flag indicates whether the caller requires ToolBuilder instances. 
+	 * 
+	 * If creating multiple copies of the ToolBuilder for a particular Tool would be problematic, it 
+	 * is the responsiblity of the IModelingModeToolExtension implementation to cache these internally, 
+	 * otherwise they will be garbage collected when the caller releases them.
+	 */
+	virtual void GetExtensionTools(const FExtensionToolQueryInfo& QueryInfo, TArray<FExtensionToolDescription>& ToolsOut) = 0;
+
+
+
+	static FName GetModularFeatureName()
+	{
+		static FName FeatureName = FName(TEXT("ModelingModeToolExtension"));
+		return FeatureName;
+	}
+};
+
+
diff --git a/Engine/Plugins/Editor/ModelingToolsEditorMode/Source/ModelingToolsEditorMode/Public/ModelingToolsActions.h b/Engine/Plugins/Editor/ModelingToolsEditorMode/Source/ModelingToolsEditorMode/Public/ModelingToolsActions.h
new file mode 100644
--- /dev/null
+++ b/Engine/Plugins/Editor/ModelingToolsEditorMode/Source/ModelingToolsEditorMode/Public/ModelingToolsActions.h
@@ -0,0 +1,87 @@
+// Copyright Epic Games, Inc. All Rights Reserved.
+
+#pragma once
+
+#include "CoreMinimal.h"
+#include "Tools/InteractiveToolsCommands.h"
+
+
+
+enum class EModelingModeActionCommands
+{
+	FocusViewToCursor
+};
+
+
+class FModelingModeActionCommands : public TCommands<FModelingModeActionCommands>
+{
+public:
+	FModelingModeActionCommands();
+
+	virtual void RegisterCommands() override;
+
+	TSharedPtr<FUICommandInfo> FocusViewCommand;
+
+	static void RegisterCommandBindings(TSharedPtr<FUICommandList> UICommandList, TFunction<void(EModelingModeActionCommands)> OnCommandExecuted);
+	static void UnRegisterCommandBindings(TSharedPtr<FUICommandList> UICommandList);
+};
+
+
+
+/**
+ * TInteractiveToolCommands implementation for this module that provides standard Editor hotkey support
+ */
+class FModelingToolActionCommands : public TInteractiveToolCommands<FModelingToolActionCommands>
+{
+public:
+	FModelingToolActionCommands();
+
+	virtual void GetToolDefaultObjectList(TArray<UInteractiveTool*>& ToolCDOs) override;
+
+	/**
+	 * interface that hides various per-tool action sets
+	 */
+
+	/**
+	 * Register all Tool command sets. Call this in module startup
+	 */
+	static void RegisterAllToolActions();
+
+	/**
+	 * Unregister all Tool command sets. Call this from module shutdown.
+	 */
+	static void UnregisterAllToolActions();
+
+	/**
+	 * Add or remove commands relevant to Tool to the given UICommandList.
+	 * Call this when the active tool changes (eg on ToolManager.OnToolStarted / OnToolEnded)
+	 * @param bUnbind if true, commands are removed, otherwise added
+	 */
+	static void UpdateToolCommandBinding(UInteractiveTool* Tool, TSharedPtr<FUICommandList> UICommandList, bool bUnbind = false);
+
+};
+
+
+// Each TCommands can only have a single (chord,action) binding, regardless of whether these
+// would ever be used at the same time. And we cannot create/register TCommands at runtime.
+// So, we have to define a separate TCommands instance for each Tool. This is unfortunate.
+
+#define DECLARE_TOOL_ACTION_COMMANDS(CommandsClassName) \
+class CommandsClassName : public TInteractiveToolCommands<CommandsClassName> \
+{\
+public:\
+	CommandsClassName();\
+	virtual void GetToolDefaultObjectList(TArray<UInteractiveTool*>& ToolCDOs) override;\
+};\
+
+
+DECLARE_TOOL_ACTION_COMMANDS(FSculptToolActionCommands);
+DECLARE_TOOL_ACTION_COMMANDS(FVertexSculptToolActionCommands);
+DECLARE_TOOL_ACTION_COMMANDS(FMeshGroupPaintToolActionCommands);
+DECLARE_TOOL_ACTION_COMMANDS(FTransformToolActionCommands);
+DECLARE_TOOL_ACTION_COMMANDS(FDrawPolygonToolActionCommands);
+DECLARE_TOOL_ACTION_COMMANDS(FMeshSelectionToolActionCommands);
+DECLARE_TOOL_ACTION_COMMANDS(FMeshPlaneCutToolActionCommands);
+DECLARE_TOOL_ACTION_COMMANDS(FEditMeshPolygonsToolActionCommands);
+DECLARE_TOOL_ACTION_COMMANDS(FDrawAndRevolveToolActionCommands);
+
diff --git a/Engine/Plugins/Editor/ModelingToolsEditorMode/Source/ModelingToolsEditorMode/Public/ModelingToolsEditorMode.h b/Engine/Plugins/Editor/ModelingToolsEditorMode/Source/ModelingToolsEditorMode/Public/ModelingToolsEditorMode.h
new file mode 100644
--- /dev/null
+++ b/Engine/Plugins/Editor/ModelingToolsEditorMode/Source/ModelingToolsEditorMode/Public/ModelingToolsEditorMode.h
@@ -0,0 +1,81 @@
+// Copyright Epic Games, Inc. All Rights Reserved.
+
+#pragma once
+
+#include "CoreMinimal.h"
+#include "Tools/LegacyEdModeWidgetHelpers.h"
+
+#include "InputState.h"
+#include "InteractiveToolManager.h"
+#include "EdModeInteractiveToolsContext.h"
+#include "ModelingToolsActions.h"
+
+#include "ModelingToolsEditorMode.generated.h"
+
+class FEditorComponentSourceFactory;
+class FUICommandList;
+class FStylusStateTracker;		// for stylus events
+
+UCLASS(Transient)
+class UModelingToolsEditorMode : public UBaseLegacyWidgetEdMode
+{
+	GENERATED_BODY()
+public:
+	const static FEditorModeID EM_ModelingToolsEditorModeId;
+
+	UModelingToolsEditorMode();
+	UModelingToolsEditorMode(FVTableHelper& Helper);
+	~UModelingToolsEditorMode();
+	////////////////
+	// UEdMode interface
+	////////////////
+
+	virtual void Tick(FEditorViewportClient* ViewportClient, float DeltaTime) override;
+
+	virtual bool ShouldDrawWidget() const override;
+	virtual bool ProcessEditDelete() override;
+	virtual bool ProcessEditCut() override;
+
+	virtual bool CanAutoSave() const override;
+
+	virtual bool ComputeBoundingBoxForViewportFocus(AActor* Actor, UPrimitiveComponent* PrimitiveComponent, FBox& InOutBox) const override;
+
+	virtual bool GetPivotForOrbit(FVector& OutPivot) const override;
+
+	/*
+	 * focus events
+	 */
+
+	// called when we "start" this editor mode (ie switch to this tab)
+	virtual void Enter() override;
+
+	// called when we "end" this editor mode (ie switch to another tab)
+	virtual void Exit() override;
+
+	//////////////////
+	// End of UEdMode interface
+	//////////////////
+
+protected:
+	virtual void BindCommands() override;
+	virtual void CreateToolkit() override;
+	virtual void OnToolStarted(UInteractiveToolManager* Manager, UInteractiveTool* Tool) override;
+	virtual void OnToolEnded(UInteractiveToolManager* Manager, UInteractiveTool* Tool) override;
+
+	FDelegateHandle MeshCreatedEventHandle;
+	FDelegateHandle TextureCreatedEventHandle;
+
+	TUniquePtr<FStylusStateTracker> StylusStateTracker;
+
+	void ModelingModeShortcutRequested(EModelingModeActionCommands Command);
+	void FocusCameraAtCursorHotkey();
+
+	void AcceptActiveToolActionOrTool();
+	void CancelActiveToolActionOrTool();
+
+	void ConfigureRealTimeViewportsOverride(bool bEnable);
+
+	// analytics tracking
+	static FDateTime LastModeStartTimestamp;
+	static FDateTime LastToolStartTimestamp;
+};
diff --git a/Engine/Plugins/Editor/ModelingToolsEditorMode/Source/ModelingToolsEditorMode/Public/ModelingToolsEditorModeModule.h b/Engine/Plugins/Editor/ModelingToolsEditorMode/Source/ModelingToolsEditorMode/Public/ModelingToolsEditorModeModule.h
new file mode 100644
--- /dev/null
+++ b/Engine/Plugins/Editor/ModelingToolsEditorMode/Source/ModelingToolsEditorMode/Public/ModelingToolsEditorModeModule.h
@@ -0,0 +1,17 @@
+// Copyright Epic Games, Inc. All Rights Reserved.
+
+#pragma once
+
+#include "CoreMinimal.h"
+#include "Modules/ModuleManager.h"
+
+class FModelingToolsEditorModeModule : public IModuleInterface
+{
+public:
+
+	/** IModuleInterface implementation */
+	virtual void StartupModule() override;
+	virtual void ShutdownModule() override;
+
+	void OnPostEngineInit();
+};
diff --git a/Engine/Plugins/Editor/ModelingToolsEditorMode/Source/ModelingToolsEditorMode/Public/ModelingToolsEditorModeSettings.h b/Engine/Plugins/Editor/ModelingToolsEditorMode/Source/ModelingToolsEditorMode/Public/ModelingToolsEditorModeSettings.h
new file mode 100644
--- /dev/null
+++ b/Engine/Plugins/Editor/ModelingToolsEditorMode/Source/ModelingToolsEditorMode/Public/ModelingToolsEditorModeSettings.h
@@ -0,0 +1,105 @@
+// Copyright Epic Games, Inc. All Rights Reserved.
+
+#pragma once
+
+#include "UObject/ObjectMacros.h"
+#include "UObject/Object.h"
+#include "ModelingToolsEditorModeSettings.generated.h"
+
+
+
+UENUM()
+enum class EModelingModeDefaultMeshObjectType
+{
+	/** Generate a new Static Mesh Asset (using Generated Asset settings below) and AStaticMeshActor */
+	StaticMeshAsset,
+	/** Generate a new AVolume Actor */
+	VolumeActor,
+	/** Generate a new ADynamicMeshActor (stored locally in the Level) */
+	DynamicMeshActor
+};
+
+
+
+UENUM()
+enum class EModelingModeAssetGenerationBehavior
+{
+	/** Generate and automatically Save new Assets on creation */
+	AutoGenerateAndAutosave,
+
+	/** Generate new Assets and mark as Modified but do not automatically Save */
+	AutoGenerateButDoNotAutosave,
+
+	/** Prompt to Save each new Asset upon Creation */
+	InteractivePromptToSave
+};
+
+
+
+UENUM()
+enum class EModelingModeAssetGenerationLocation
+{
+	/** All generated assets will be stored in an AutoGenerated folder that is located relative to the World they are being saved in */
+	AutoGeneratedWorldRelativeAssetPath,
+
+	/** All generated assets will be stored in a root-level AutoGenerated folder */
+	AutoGeneratedGlobalAssetPath,
+
+	/** Generated assets will be stored in the currently-visible Asset Browser folder if available, otherwise at the Auto Generated Asset Path */
+	CurrentAssetBrowserPathIfAvailable
+};
+
+
+/**
+ * Settings for the Modeling Tools Editor Mode plug-in.
+ */
+UCLASS(config=Editor)
+class MODELINGTOOLSEDITORMODE_API UModelingToolsEditorModeSettings
+	: public UObject
+{
+	GENERATED_BODY()
+public:
+
+	/** Enable/Disable the options to emit Dynamic Mesh Actors in Modeling Mode Tools */
+	UPROPERTY(config, EditAnywhere, Category = "Modeling Mode")
+	bool bEnableDynamicMeshActors = false;
+
+	/** What type of Mesh Object should Output Type Setting default to in Modeling Mode Tools (takes effect after Editor restart) */
+	UPROPERTY(config, EditAnywhere, Category = "Modeling Mode")
+	EModelingModeDefaultMeshObjectType DefaultMeshObjectType = EModelingModeDefaultMeshObjectType::StaticMeshAsset;
+
+	/** How should Assets auto-generated by Modeling Tools be handled in terms of saving, naming, etc */
+	UPROPERTY(config, EditAnywhere, Category = "Modeling Mode|Generated Assets")
+	EModelingModeAssetGenerationBehavior AssetGenerationMode = EModelingModeAssetGenerationBehavior::AutoGenerateButDoNotAutosave;
+
+	/** Where should Assets auto-generated by Modeling Tools be stored by default */
+	UPROPERTY(config, EditAnywhere, Category = "Modeling Mode|Generated Assets")
+	EModelingModeAssetGenerationLocation AssetGenerationLocation = EModelingModeAssetGenerationLocation::AutoGeneratedWorldRelativeAssetPath;
+
+	/** Assets auto-generated by Modeling Tools are stored at this path, relative to the parent path defined by the Location. Set to an empty string to disable. */
+	UPROPERTY(config, EditAnywhere, Category = "Modeling Mode|Generated Assets")
+	FString AutoGeneratedAssetPath = TEXT("_GENERATED");
+
+	/** If true, Auto-Generated Assets created in an unsaved Level will be stored relative to top-level /Game folder, otherwise they will be stored in /Temp and cannot be saved until they are explicitly moved to a permanent location */
+	UPROPERTY(config, EditAnywhere, Category = "Modeling Mode|Generated Assets")
+	bool bStoreUnsavedLevelAssetsInTopLevelGameFolder = true;
+
+	/** If true, Autogenerated Assets are stored in per-user folders below the Autogen path */
+	UPROPERTY(config, EditAnywhere, Category = "Modeling Mode|Generated Assets")
+	bool bUsePerUserAutogenSubfolder = true;
+
+	/** If true, Autogenerated Assets have a short random string generated and appended to their name */
+	UPROPERTY(config, EditAnywhere, Category = "Modeling Mode|Generated Assets")
+	bool bAppendRandomStringToName = true;
+
+public:
+
+	DECLARE_MULTICAST_DELEGATE_TwoParams(UModelingToolsEditorModeSettingsModified, UObject*, FProperty*);
+	UModelingToolsEditorModeSettingsModified OnModified;
+
+	void PostEditChangeProperty(FPropertyChangedEvent& PropertyChangedEvent)
+	{
+		OnModified.Broadcast(this, PropertyChangedEvent.Property);
+	}
+
+};
diff --git a/Engine/Plugins/Editor/ModelingToolsEditorMode/Source/ModelingToolsEditorMode/Public/ModelingToolsEditorModeStyle.h b/Engine/Plugins/Editor/ModelingToolsEditorMode/Source/ModelingToolsEditorMode/Public/ModelingToolsEditorModeStyle.h
new file mode 100644
--- /dev/null
+++ b/Engine/Plugins/Editor/ModelingToolsEditorMode/Source/ModelingToolsEditorMode/Public/ModelingToolsEditorModeStyle.h
@@ -0,0 +1,24 @@
+// Copyright Epic Games, Inc. All Rights Reserved.
+
+#pragma once
+
+#include "CoreMinimal.h"
+#include "Styling/ISlateStyle.h"
+#include "Styling/SlateStyle.h"
+
+class MODELINGTOOLSEDITORMODE_API FModelingToolsEditorModeStyle
+{
+public:
+	static void Initialize();
+
+	static void Shutdown();
+
+	static TSharedPtr< class ISlateStyle > Get();
+
+	static FName GetStyleSetName();
+private:
+	static FString InContent(const FString& RelativePath, const ANSICHAR* Extension);
+
+private:
+	static TSharedPtr< class FSlateStyleSet > StyleSet;
+};
diff --git a/Engine/Plugins/Editor/ModelingToolsEditorMode/Source/ModelingToolsEditorMode/Public/ModelingToolsEditorModeToolkit.h b/Engine/Plugins/Editor/ModelingToolsEditorMode/Source/ModelingToolsEditorMode/Public/ModelingToolsEditorModeToolkit.h
new file mode 100644
--- /dev/null
+++ b/Engine/Plugins/Editor/ModelingToolsEditorMode/Source/ModelingToolsEditorMode/Public/ModelingToolsEditorModeToolkit.h
@@ -0,0 +1,115 @@
+// Copyright Epic Games, Inc. All Rights Reserved.
+
+#pragma once
+
+#include "CoreMinimal.h"
+#include "Toolkits/BaseToolkit.h"
+#include "InteractiveTool.h"
+#include "ModelingToolsEditorMode.h"
+#include "Widgets/Input/SComboBox.h"
+#include "Widgets/Input/STextComboBox.h"
+#include "Widgets/SBoxPanel.h"
+#include "StatusBarSubsystem.h"
+
+
+class IDetailsView;
+class SButton;
+class STextBlock;
+
+class FModelingToolsEditorModeToolkit : public FModeToolkit
+{
+public:
+
+	FModelingToolsEditorModeToolkit();
+	~FModelingToolsEditorModeToolkit();
+	
+	/** FModeToolkit interface */
+	virtual void Init(const TSharedPtr<IToolkitHost>& InitToolkitHost, TWeakObjectPtr<UEdMode> InOwningMode) override;
+
+	/** IToolkit interface */
+	virtual FName GetToolkitFName() const override;
+	virtual FText GetBaseToolkitName() const override;
+	virtual TSharedPtr<class SWidget> GetInlineContent() const override { return ToolkitWidget; }
+
+	// initialize toolkit widgets that need to wait until mode is initialized/entered
+	virtual void InitializeAfterModeSetup();
+
+	// set/clear notification message area
+	virtual void PostNotification(const FText& Message);
+	virtual void ClearNotification();
+
+	// set/clear warning message area
+	virtual void PostWarning(const FText& Message);
+	virtual void ClearWarning();
+
+	/** Returns the Mode specific tabs in the mode toolbar **/ 
+	virtual void GetToolPaletteNames(TArray<FName>& InPaletteName) const;
+	virtual FText GetToolPaletteDisplayName(FName PaletteName) const; 
+	virtual void BuildToolPalette(FName PaletteName, class FToolBarBuilder& ToolbarBuilder);
+	virtual void OnToolPaletteChanged(FName PaletteName) override;
+	virtual bool HasIntegratedToolPalettes() const { return false; }
+	virtual bool HasExclusiveToolPalettes() const { return false; }
+
+	virtual FText GetActiveToolDisplayName() const override { return ActiveToolName; }
+	virtual FText GetActiveToolMessage() const override { return ActiveToolMessage; }
+
+	virtual void BuildToolPalette_Standard(FName PaletteName, class FToolBarBuilder& ToolbarBuilder);
+	virtual void BuildToolPalette_Experimental(FName PaletteName, class FToolBarBuilder& ToolbarBuilder);
+
+	virtual void EnableShowRealtimeWarning(bool bEnable);
+
+	virtual void OnToolStarted(UInteractiveToolManager* Manager, UInteractiveTool* Tool) override;
+	virtual void OnToolEnded(UInteractiveToolManager* Manager, UInteractiveTool* Tool) override;
+
+	virtual void CustomizeModeDetailsViewArgs(FDetailsViewArgs& ArgsInOut) override;
+
+	void OnActiveViewportChanged(TSharedPtr<IAssetViewport>, TSharedPtr<IAssetViewport> );
+
+
+private:
+	const static TArray<FName> PaletteNames_Standard;
+	const static TArray<FName> PaletteNames_Experimental;
+
+	FText ActiveToolName;
+	FText ActiveToolMessage;
+	FStatusBarMessageHandle ActiveToolMessageHandle;
+	const FSlateBrush* ActiveToolIcon;
+
+	TSharedPtr<SWidget> ToolkitWidget;
+	void UpdateActiveToolProperties();
+	void InvalidateCachedDetailPanelState(UObject* ChangedObject);
+
+	TSharedPtr<SWidget> ViewportOverlayWidget;
+
+	TSharedPtr<STextBlock> ModeWarningArea;
+	TSharedPtr<STextBlock> ModeHeaderArea;
+	TSharedPtr<STextBlock> ToolWarningArea;
+	TSharedPtr<SButton> AcceptButton;
+	TSharedPtr<SButton> CancelButton;
+	TSharedPtr<SButton> CompletedButton;
+
+	TSharedPtr<SWidget> MakeAssetConfigPanel();
+
+	bool bShowRealtimeWarning = false;
+	void UpdateShowWarnings();
+
+	TArray<TSharedPtr<FString>> AssetLocationModes;
+	TArray<TSharedPtr<FString>> AssetSaveModes;
+	TSharedPtr<STextComboBox> AssetLocationMode;
+	TSharedPtr<STextComboBox> AssetSaveMode;
+	void UpdateAssetLocationMode(TSharedPtr<FString> NewString);
+	void UpdateAssetSaveMode(TSharedPtr<FString> NewString);
+	void UpdateAssetPanelFromSettings();
+	void UpdateObjectCreationOptionsFromSettings();
+	void OnAssetSettingsModified();
+	FDelegateHandle AssetSettingsModifiedHandle;
+	void OnShowAssetSettings();
+
+	TArray<TSharedPtr<FString>> AssetLODModes;
+	TSharedPtr<STextBlock> AssetLODModeLabel;
+	TSharedPtr<STextComboBox> AssetLODMode;
+
+	bool bFirstInitializeAfterModeSetup = true;
+
+	// custom accept/cancel/complete handlers
+};
diff --git a/Engine/Plugins/Editor/ModelingToolsEditorMode/Source/ModelingToolsEditorMode/Public/ModelingToolsManagerActions.h b/Engine/Plugins/Editor/ModelingToolsEditorMode/Source/ModelingToolsEditorMode/Public/ModelingToolsManagerActions.h
new file mode 100644
--- /dev/null
+++ b/Engine/Plugins/Editor/ModelingToolsEditorMode/Source/ModelingToolsEditorMode/Public/ModelingToolsManagerActions.h
@@ -0,0 +1,136 @@
+// Copyright Epic Games, Inc. All Rights Reserved.
+
+#pragma once
+
+#include "CoreMinimal.h"
+#include "Framework/Commands/Commands.h"
+
+
+/**
+ * TInteractiveToolCommands implementation for this module that provides standard Editor hotkey support
+ */
+class MODELINGTOOLSEDITORMODE_API FModelingToolsManagerCommands : public TCommands<FModelingToolsManagerCommands>
+{
+public:
+	FModelingToolsManagerCommands();
+
+
+	TSharedPtr<FUICommandInfo> BeginAddBoxPrimitiveTool;
+	TSharedPtr<FUICommandInfo> BeginAddCylinderPrimitiveTool;
+	TSharedPtr<FUICommandInfo> BeginAddConePrimitiveTool;
+	TSharedPtr<FUICommandInfo> BeginAddArrowPrimitiveTool;
+	TSharedPtr<FUICommandInfo> BeginAddRectanglePrimitiveTool;
+	TSharedPtr<FUICommandInfo> BeginAddDiscPrimitiveTool;
+	TSharedPtr<FUICommandInfo> BeginAddTorusPrimitiveTool;
+	TSharedPtr<FUICommandInfo> BeginAddSpherePrimitiveTool;
+	TSharedPtr<FUICommandInfo> BeginAddStairsPrimitiveTool;
+
+	TSharedPtr<FUICommandInfo> BeginAddPatchTool;
+	TSharedPtr<FUICommandInfo> BeginRevolveBoundaryTool;
+	TSharedPtr<FUICommandInfo> BeginDrawPolygonTool;
+	TSharedPtr<FUICommandInfo> BeginDrawPolyPathTool;
+	TSharedPtr<FUICommandInfo> BeginDrawAndRevolveTool;
+	TSharedPtr<FUICommandInfo> BeginShapeSprayTool;
+
+	TSharedPtr<FUICommandInfo> BeginSculptMeshTool;
+	TSharedPtr<FUICommandInfo> BeginPolyEditTool;
+	TSharedPtr<FUICommandInfo> BeginSubdividePolyTool;
+	TSharedPtr<FUICommandInfo> BeginTriEditTool;
+	TSharedPtr<FUICommandInfo> BeginPolyDeformTool;
+	TSharedPtr<FUICommandInfo> BeginSmoothMeshTool;
+	TSharedPtr<FUICommandInfo> BeginOffsetMeshTool;
+	TSharedPtr<FUICommandInfo> BeginDisplaceMeshTool;
+	TSharedPtr<FUICommandInfo> BeginMeshSpaceDeformerTool;
+	TSharedPtr<FUICommandInfo> BeginTransformMeshesTool;
+	TSharedPtr<FUICommandInfo> BeginEditPivotTool;
+	TSharedPtr<FUICommandInfo> BeginBakeTransformTool;
+	TSharedPtr<FUICommandInfo> BeginCombineMeshesTool;
+	TSharedPtr<FUICommandInfo> BeginDuplicateMeshesTool;
+	TSharedPtr<FUICommandInfo> BeginAlignObjectsTool;
+	TSharedPtr<FUICommandInfo> BeginTransferMeshTool;
+	TSharedPtr<FUICommandInfo> BeginConvertMeshesTool;
+	TSharedPtr<FUICommandInfo> BeginSplitMeshesTool;
+
+	TSharedPtr<FUICommandInfo> BeginRemeshSculptMeshTool;
+	TSharedPtr<FUICommandInfo> BeginRemeshMeshTool;
+	TSharedPtr<FUICommandInfo> BeginProjectToTargetTool;
+	TSharedPtr<FUICommandInfo> BeginSimplifyMeshTool;
+	TSharedPtr<FUICommandInfo> BeginEditNormalsTool;
+	TSharedPtr<FUICommandInfo> BeginEditTangentsTool;
+	TSharedPtr<FUICommandInfo> BeginRemoveOccludedTrianglesTool;
+	TSharedPtr<FUICommandInfo> BeginUVProjectionTool;
+	TSharedPtr<FUICommandInfo> BeginUVLayoutTool;
+	TSharedPtr<FUICommandInfo> BeginPlaneCutTool;
+	TSharedPtr<FUICommandInfo> BeginMirrorTool;
+	TSharedPtr<FUICommandInfo> BeginHoleFillTool;
+	TSharedPtr<FUICommandInfo> BeginLatticeDeformerTool;
+	TSharedPtr<FUICommandInfo> BeginPolygonCutTool;
+	TSharedPtr<FUICommandInfo> BeginVoxelMergeTool;
+	TSharedPtr<FUICommandInfo> BeginVoxelBooleanTool;
+	TSharedPtr<FUICommandInfo> BeginVoxelSolidifyTool;
+	TSharedPtr<FUICommandInfo> BeginVoxelBlendTool;
+	TSharedPtr<FUICommandInfo> BeginVoxelMorphologyTool;
+	TSharedPtr<FUICommandInfo> BeginSelfUnionTool;
+	TSharedPtr<FUICommandInfo> BeginMeshBooleanTool;
+	TSharedPtr<FUICommandInfo> BeginMeshTrimTool;
+	TSharedPtr<FUICommandInfo> BeginCutMeshWithMeshTool;
+	TSharedPtr<FUICommandInfo> BeginMeshSelectionTool;
+	TSharedPtr<FUICommandInfo> BeginBspConversionTool;
+	TSharedPtr<FUICommandInfo> BeginMeshToVolumeTool;
+	TSharedPtr<FUICommandInfo> BeginVolumeToMeshTool;
+	TSharedPtr<FUICommandInfo> BeginMeshGroupPaintTool;
+
+	TSharedPtr<FUICommandInfo> BeginPhysicsInspectorTool;
+	TSharedPtr<FUICommandInfo> BeginSetCollisionGeometryTool;
+	TSharedPtr<FUICommandInfo> BeginEditCollisionGeometryTool;
+	TSharedPtr<FUICommandInfo> BeginExtractCollisionGeometryTool;
+
+	TSharedPtr<FUICommandInfo> BeginMeshInspectorTool;
+	TSharedPtr<FUICommandInfo> BeginGlobalUVGenerateTool;
+	TSharedPtr<FUICommandInfo> BeginGroupUVGenerateTool;
+	TSharedPtr<FUICommandInfo> BeginWeldEdgesTool;
+	TSharedPtr<FUICommandInfo> BeginPolyGroupsTool;
+	TSharedPtr<FUICommandInfo> BeginEditMeshMaterialsTool;
+	TSharedPtr<FUICommandInfo> BeginTransformUVIslandsTool;
+	TSharedPtr<FUICommandInfo> BeginMeshAttributePaintTool;
+	TSharedPtr<FUICommandInfo> BeginAttributeEditorTool;
+	TSharedPtr<FUICommandInfo> BeginBakeMeshAttributeMapsTool;
+	TSharedPtr<FUICommandInfo> BeginBakeMeshAttributeVertexTool;
+	TSharedPtr<FUICommandInfo> BeginUVSeamEditTool;
+
+	TSharedPtr<FUICommandInfo> BeginSkinWeightsPaintTool;
+	TSharedPtr<FUICommandInfo> BeginSkinWeightsBindingTool;
+	
+	TSharedPtr<FUICommandInfo> BeginGroomToMeshTool;
+	TSharedPtr<FUICommandInfo> BeginGenerateLODMeshesTool;
+	TSharedPtr<FUICommandInfo> BeginGroomCardsToMeshTool;
+	TSharedPtr<FUICommandInfo> BeginGroomMeshToCardsTool;
+	TSharedPtr<FUICommandInfo> BeginGroomCardsEditorTool;
+
+	TSharedPtr<FUICommandInfo> BeginLODManagerTool;
+	TSharedPtr<FUICommandInfo> BeginGenerateStaticMeshLODAssetTool;
+
+	TSharedPtr<FUICommandInfo> BeginPolyModelTool_FaceSelect;
+	TSharedPtr<FUICommandInfo> BeginPolyModelTool_EdgeSelect;
+	TSharedPtr<FUICommandInfo> BeginPolyModelTool_VertexSelect;
+	TSharedPtr<FUICommandInfo> BeginPolyModelTool_AllSelect;
+	TSharedPtr<FUICommandInfo> BeginPolyModelTool_LoopSelect;
+	TSharedPtr<FUICommandInfo> BeginPolyModelTool_RingSelect;
+
+	TSharedPtr<FUICommandInfo> BeginPolyModelTool_Extrude;
+	TSharedPtr<FUICommandInfo> BeginPolyModelTool_InsetOutset;
+	TSharedPtr<FUICommandInfo> BeginPolyModelTool_CutFaces;
+
+
+	TSharedPtr<FUICommandInfo> AcceptActiveTool;
+	TSharedPtr<FUICommandInfo> CancelActiveTool;
+	TSharedPtr<FUICommandInfo> CompleteActiveTool;
+
+	TSharedPtr<FUICommandInfo> AcceptOrCompleteActiveTool;
+	TSharedPtr<FUICommandInfo> CancelOrCompleteActiveTool;
+
+	/**
+	 * Initialize commands
+	 */
+	virtual void RegisterCommands() override;
+};
diff --git a/Engine/Plugins/Experimental/ModelingToolsEditorMode/ModelingToolsEditorMode.uplugin b/Engine/Plugins/Experimental/ModelingToolsEditorMode/ModelingToolsEditorMode.uplugin
deleted file mode 100644
diff --git a/Engine/Plugins/Experimental/ModelingToolsEditorMode/Source/ModelingToolsEditorMode/ModelingToolsEditorMode.Build.cs b/Engine/Plugins/Experimental/ModelingToolsEditorMode/Source/ModelingToolsEditorMode/ModelingToolsEditorMode.Build.cs
deleted file mode 100644
diff --git a/Engine/Plugins/Experimental/ModelingToolsEditorMode/Source/ModelingToolsEditorMode/Private/DetailsCustomizations/ModelingToolsAxisFilterCustomization.cpp b/Engine/Plugins/Experimental/ModelingToolsEditorMode/Source/ModelingToolsEditorMode/Private/DetailsCustomizations/ModelingToolsAxisFilterCustomization.cpp
deleted file mode 100644
diff --git a/Engine/Plugins/Experimental/ModelingToolsEditorMode/Source/ModelingToolsEditorMode/Private/ModelingModeAssetUtils.cpp b/Engine/Plugins/Experimental/ModelingToolsEditorMode/Source/ModelingToolsEditorMode/Private/ModelingModeAssetUtils.cpp
deleted file mode 100644
diff --git a/Engine/Plugins/Experimental/ModelingToolsEditorMode/Source/ModelingToolsEditorMode/Private/ModelingToolsActions.cpp b/Engine/Plugins/Experimental/ModelingToolsEditorMode/Source/ModelingToolsEditorMode/Private/ModelingToolsActions.cpp
deleted file mode 100644
diff --git a/Engine/Plugins/Experimental/ModelingToolsEditorMode/Source/ModelingToolsEditorMode/Private/ModelingToolsEditorMode.cpp b/Engine/Plugins/Experimental/ModelingToolsEditorMode/Source/ModelingToolsEditorMode/Private/ModelingToolsEditorMode.cpp
deleted file mode 100644
diff --git a/Engine/Plugins/Experimental/ModelingToolsEditorMode/Source/ModelingToolsEditorMode/Private/ModelingToolsEditorModeModule.cpp b/Engine/Plugins/Experimental/ModelingToolsEditorMode/Source/ModelingToolsEditorMode/Private/ModelingToolsEditorModeModule.cpp
deleted file mode 100644
diff --git a/Engine/Plugins/Experimental/ModelingToolsEditorMode/Source/ModelingToolsEditorMode/Private/ModelingToolsEditorModeStyle.cpp b/Engine/Plugins/Experimental/ModelingToolsEditorMode/Source/ModelingToolsEditorMode/Private/ModelingToolsEditorModeStyle.cpp
deleted file mode 100644
diff --git a/Engine/Plugins/Experimental/ModelingToolsEditorMode/Source/ModelingToolsEditorMode/Private/ModelingToolsEditorModeToolkit.cpp b/Engine/Plugins/Experimental/ModelingToolsEditorMode/Source/ModelingToolsEditorMode/Private/ModelingToolsEditorModeToolkit.cpp
deleted file mode 100644
diff --git a/Engine/Plugins/Experimental/ModelingToolsEditorMode/Source/ModelingToolsEditorMode/Private/ModelingToolsManagerActions.cpp b/Engine/Plugins/Experimental/ModelingToolsEditorMode/Source/ModelingToolsEditorMode/Private/ModelingToolsManagerActions.cpp
deleted file mode 100644
diff --git a/Engine/Plugins/Experimental/ModelingToolsEditorMode/Source/ModelingToolsEditorMode/Public/DetailsCustomizations/ModelingToolPropertyCustomizations.h b/Engine/Plugins/Experimental/ModelingToolsEditorMode/Source/ModelingToolsEditorMode/Public/DetailsCustomizations/ModelingToolPropertyCustomizations.h
deleted file mode 100644
diff --git a/Engine/Plugins/Experimental/ModelingToolsEditorMode/Source/ModelingToolsEditorMode/Public/ModelingModeAssetUtils.h b/Engine/Plugins/Experimental/ModelingToolsEditorMode/Source/ModelingToolsEditorMode/Public/ModelingModeAssetUtils.h
deleted file mode 100644
diff --git a/Engine/Plugins/Experimental/ModelingToolsEditorMode/Source/ModelingToolsEditorMode/Public/ModelingModeToolExtensions.h b/Engine/Plugins/Experimental/ModelingToolsEditorMode/Source/ModelingToolsEditorMode/Public/ModelingModeToolExtensions.h
deleted file mode 100644
diff --git a/Engine/Plugins/Experimental/ModelingToolsEditorMode/Source/ModelingToolsEditorMode/Public/ModelingToolsActions.h b/Engine/Plugins/Experimental/ModelingToolsEditorMode/Source/ModelingToolsEditorMode/Public/ModelingToolsActions.h
deleted file mode 100644
diff --git a/Engine/Plugins/Experimental/ModelingToolsEditorMode/Source/ModelingToolsEditorMode/Public/ModelingToolsEditorMode.h b/Engine/Plugins/Experimental/ModelingToolsEditorMode/Source/ModelingToolsEditorMode/Public/ModelingToolsEditorMode.h
deleted file mode 100644
diff --git a/Engine/Plugins/Experimental/ModelingToolsEditorMode/Source/ModelingToolsEditorMode/Public/ModelingToolsEditorModeModule.h b/Engine/Plugins/Experimental/ModelingToolsEditorMode/Source/ModelingToolsEditorMode/Public/ModelingToolsEditorModeModule.h
deleted file mode 100644
diff --git a/Engine/Plugins/Experimental/ModelingToolsEditorMode/Source/ModelingToolsEditorMode/Public/ModelingToolsEditorModeSettings.h b/Engine/Plugins/Experimental/ModelingToolsEditorMode/Source/ModelingToolsEditorMode/Public/ModelingToolsEditorModeSettings.h
deleted file mode 100644
diff --git a/Engine/Plugins/Experimental/ModelingToolsEditorMode/Source/ModelingToolsEditorMode/Public/ModelingToolsEditorModeStyle.h b/Engine/Plugins/Experimental/ModelingToolsEditorMode/Source/ModelingToolsEditorMode/Public/ModelingToolsEditorModeStyle.h
deleted file mode 100644
diff --git a/Engine/Plugins/Experimental/ModelingToolsEditorMode/Source/ModelingToolsEditorMode/Public/ModelingToolsEditorModeToolkit.h b/Engine/Plugins/Experimental/ModelingToolsEditorMode/Source/ModelingToolsEditorMode/Public/ModelingToolsEditorModeToolkit.h
deleted file mode 100644
diff --git a/Engine/Plugins/Experimental/ModelingToolsEditorMode/Source/ModelingToolsEditorMode/Public/ModelingToolsManagerActions.h b/Engine/Plugins/Experimental/ModelingToolsEditorMode/Source/ModelingToolsEditorMode/Public/ModelingToolsManagerActions.h
deleted file mode 100644
diff --git a/Engine/Source/Runtime/AppFramework/Private/Framework/Testing/SStarshipGallery.cpp b/Engine/Source/Runtime/AppFramework/Private/Framework/Testing/SStarshipGallery.cpp
--- a/Engine/Source/Runtime/AppFramework/Private/Framework/Testing/SStarshipGallery.cpp
+++ b/Engine/Source/Runtime/AppFramework/Private/Framework/Testing/SStarshipGallery.cpp
@@ -465,7 +465,7 @@ public:
 
                 +SVerticalBox::Slot().AutoHeight()[ GenerateIconLibrary(NSLOCTEXT("StarshipGallery", "PaintIconTitle", "Paint"), "Content/Editor/Slate/Icons/Paint")]
                 +SVerticalBox::Slot().AutoHeight()[ GenerateIconLibrary(NSLOCTEXT("StarshipGallery", "LandscapeIconTitle", "Landscape"), "Content/Editor/Slate/Icons/Landscape")]
-                +SVerticalBox::Slot().AutoHeight()[ GenerateIconLibrary(NSLOCTEXT("StarshipGallery", "ModelingIconTitle", "Modeling"), "/Plugins/Experimental/ModelingToolsEditorMode/Content/Icons")]
+                +SVerticalBox::Slot().AutoHeight()[ GenerateIconLibrary(NSLOCTEXT("StarshipGallery", "ModelingIconTitle", "Modeling"), "/Plugins/Editor/ModelingToolsEditorMode/Content/Icons")]
                 +SVerticalBox::Slot().AutoHeight()[ GenerateIconLibrary(NSLOCTEXT("StarshipGallery", "FractureIconTitle", "Fracture"), "/Plugins/Experimental/ChaosEditor/Content")]
                 +SVerticalBox::Slot().AutoHeight()[ GenerateIconLibrary(NSLOCTEXT("StarshipGallery", "CurveEditorIconTitle", "CurveEditor"), "Content/Editor/Slate/GenericCurveEditor/Icons")]
                 +SVerticalBox::Slot().AutoHeight()[ GenerateIconLibrary(NSLOCTEXT("StarshipGallery", "GeneralIconTitle", "General"), "Content/Editor/Slate/Icons/GeneralTools")]
diff --git a/Engine/Source/Runtime/Experimental/InteractiveToolsFramework/Public/InteractiveTool.h b/Engine/Source/Runtime/Experimental/InteractiveToolsFramework/Public/InteractiveTool.h
--- a/Engine/Source/Runtime/Experimental/InteractiveToolsFramework/Public/InteractiveTool.h
+++ b/Engine/Source/Runtime/Experimental/InteractiveToolsFramework/Public/InteractiveTool.h
@@ -562,7 +562,7 @@ protected:
 	/**
 	 * Override this function to register the set of Actions this Tool supports, using FInteractiveToolActionSet::RegisterAction.
 	 * Note that for the actions to be triggered, you will also need to add corresponding registration per tool
-	 *  -- see Engine\Plugins\Experimental\ModelingToolsEditorMode\Source\ModelingToolsEditorMode\Public\ModelingToolsActions.h for examples
+	 *  -- see Engine\Plugins\Editor\ModelingToolsEditorMode\Source\ModelingToolsEditorMode\Public\ModelingToolsActions.h for examples
 	 */
 	virtual void RegisterActions(FInteractiveToolActionSet& ActionSet);
 
