ModelingTools: MeshModelingToolset and MeshModelingToolsetExp plugins have modules that are Runtime-capable but also used by Editor-only ModelingToolsEditorMode. Enabling ModelingToolsEditorMode in a .uproject, even with WhitelistTargets configured to Editor, will result in the built game including the Runtime modules from MeshModelingToolset/Exp (which is about 5mb in the binary). This change modifies the .uplugin files for the various Editor Modes to also mark their dependencies as whitelisted for Editor. Also marked the PlanarCut plugin as whitelisted for Editor in ChaosEditor.uplugin.

MeshModelingToolset and PlanarCut plugins can still be explicitly included in the .uproject to enable them for Runtime use by licensees (tested this in ShooterGame in Shipping configurations).

#rnx

diff --git a/Engine/Plugins/Editor/ModelingToolsEditorMode/ModelingToolsEditorMode.uplugin b/Engine/Plugins/Editor/ModelingToolsEditorMode/ModelingToolsEditorMode.uplugin
--- a/Engine/Plugins/Editor/ModelingToolsEditorMode/ModelingToolsEditorMode.uplugin
+++ b/Engine/Plugins/Editor/ModelingToolsEditorMode/ModelingToolsEditorMode.uplugin
@@ -24,15 +24,18 @@
 	"Plugins": [
 		{
 			"Name": "MeshModelingToolset",
-			"Enabled": true
+			"Enabled": true,
+			"WhitelistTargets": [ "Editor" ]
 		},
 		{
 			"Name": "MeshModelingToolsetExp",
-			"Enabled": true
+			"Enabled": true,
+			"WhitelistTargets": [ "Editor" ]
 		},
 		{
 			"Name": "MeshLODToolset",
-			"Enabled": true
+			"Enabled": true,
+			"WhitelistTargets": [ "Editor" ]
 		},
 		{
 			"Name": "StylusInput",
diff --git a/Engine/Plugins/Experimental/ChaosEditor/ChaosEditor.uplugin b/Engine/Plugins/Experimental/ChaosEditor/ChaosEditor.uplugin
--- a/Engine/Plugins/Experimental/ChaosEditor/ChaosEditor.uplugin
+++ b/Engine/Plugins/Experimental/ChaosEditor/ChaosEditor.uplugin
@@ -24,7 +24,8 @@
 	"Plugins": [
 		{
 			"Name": "PlanarCut",
-			"Enabled": true
+			"Enabled": true,
+			"WhitelistTargets": [ "Editor" ]
 		},
 		{
 			"Name": "GeometryCollectionPlugin",
@@ -32,7 +33,8 @@
 		},
 		{
 			"Name": "MeshModelingToolsetExp",
-			"Enabled": true
+			"Enabled": true,
+			"WhitelistTargets": [ "Editor" ]
 		},
 		{
 			"Name": "EditorScriptingUtilities",
diff --git a/Engine/Plugins/Experimental/MeshLODToolset/MeshLODToolset.uplugin b/Engine/Plugins/Experimental/MeshLODToolset/MeshLODToolset.uplugin
--- a/Engine/Plugins/Experimental/MeshLODToolset/MeshLODToolset.uplugin
+++ b/Engine/Plugins/Experimental/MeshLODToolset/MeshLODToolset.uplugin
@@ -33,7 +33,8 @@
 		},
 		{
 			"Name": "MeshModelingToolsetExp",
-			"Enabled": true
+			"Enabled": true,
+			"WhitelistTargets": [ "Editor" ]
 		},
 		{
 			"Name": "EditorScriptingUtilities",
diff --git a/Engine/Plugins/Experimental/UVEditor/UVEditor.uplugin b/Engine/Plugins/Experimental/UVEditor/UVEditor.uplugin
--- a/Engine/Plugins/Experimental/UVEditor/UVEditor.uplugin
+++ b/Engine/Plugins/Experimental/UVEditor/UVEditor.uplugin
@@ -33,7 +33,8 @@
 		},
 		{
 			"Name": "MeshModelingToolsetExp",
-			"Enabled": true
+			"Enabled": true,
+			"WhitelistTargets": [ "Editor" ]
 		}
 	]
 }
\ No newline at end of file
