MeshModelingTools: Allow hotkeys to work in subclasses of tools whose parents also use hotkeys. Add hotkeys to EdMat tool. Fix largest area/tri count buttons being flipped in terms of what they do.

#rnx

diff --git a/Engine/Plugins/Editor/ModelingToolsEditorMode/Source/ModelingToolsEditorMode/Private/ModelingToolsActions.cpp b/Engine/Plugins/Editor/ModelingToolsEditorMode/Source/ModelingToolsEditorMode/Private/ModelingToolsActions.cpp
--- a/Engine/Plugins/Editor/ModelingToolsEditorMode/Source/ModelingToolsEditorMode/Private/ModelingToolsActions.cpp
+++ b/Engine/Plugins/Editor/ModelingToolsEditorMode/Source/ModelingToolsEditorMode/Private/ModelingToolsActions.cpp
@@ -3,6 +3,7 @@
 #include "ModelingToolsActions.h"
 #include "EditorStyleSet.h"
 #include "DynamicMeshSculptTool.h"
+#include "EditMeshMaterialsTool.h"
 #include "MeshVertexSculptTool.h"
 #include "MeshGroupPaintTool.h"
 #include "MeshInspectorTool.h"
@@ -86,6 +87,7 @@ void FModelingToolActionCommands::RegisterAllToolActions()
 	FDrawPolygonToolActionCommands::Register();
 	FTransformToolActionCommands::Register();
 	FMeshSelectionToolActionCommands::Register();
+	FEditMeshMaterialsToolActionCommands::Register();
 	FMeshPlaneCutToolActionCommands::Register();
 	FEditMeshPolygonsToolActionCommands::Register();
 	FDrawAndRevolveToolActionCommands::Register();
@@ -100,6 +102,7 @@ void FModelingToolActionCommands::UnregisterAllToolActions()
 	FDrawPolygonToolActionCommands::Unregister();
 	FTransformToolActionCommands::Unregister();
 	FMeshSelectionToolActionCommands::Unregister();
+	FEditMeshMaterialsToolActionCommands::Unregister();
 	FMeshPlaneCutToolActionCommands::Unregister();
 	FEditMeshPolygonsToolActionCommands::Unregister();
 	FDrawAndRevolveToolActionCommands::Unregister();
@@ -112,39 +115,43 @@ void FModelingToolActionCommands::UpdateToolCommandBinding(UInteractiveTool* Too
 #define UPDATE_BINDING(CommandsType)  if (!bUnbind) CommandsType::Get().BindCommandsForCurrentTool(UICommandList, Tool); else CommandsType::Get().UnbindActiveCommands(UICommandList);
 
 
-	if (Cast<UTransformMeshesTool>(Tool) != nullptr)
+	if (ExactCast<UTransformMeshesTool>(Tool) != nullptr)
 	{
 		UPDATE_BINDING(FTransformToolActionCommands);
 	}
-	else if (Cast<UDynamicMeshSculptTool>(Tool) != nullptr)
+	else if (ExactCast<UDynamicMeshSculptTool>(Tool) != nullptr)
 	{
 		UPDATE_BINDING(FSculptToolActionCommands);
 	}
-	else if (Cast<UMeshVertexSculptTool>(Tool) != nullptr)
+	else if (ExactCast<UMeshVertexSculptTool>(Tool) != nullptr)
 	{
 		UPDATE_BINDING(FVertexSculptToolActionCommands);
 	}
-	else if (Cast<UMeshGroupPaintTool>(Tool) != nullptr)
+	else if (ExactCast<UMeshGroupPaintTool>(Tool) != nullptr)
 	{
 		UPDATE_BINDING(FMeshGroupPaintToolActionCommands);
 	}
-	else if (Cast<UDrawPolygonTool>(Tool) != nullptr)
+	else if (ExactCast<UDrawPolygonTool>(Tool) != nullptr)
 	{
 		UPDATE_BINDING(FDrawPolygonToolActionCommands);
 	}
-	else if (Cast<UMeshSelectionTool>(Tool) != nullptr)
+	else if (ExactCast<UMeshSelectionTool>(Tool) != nullptr)
 	{
 		UPDATE_BINDING(FMeshSelectionToolActionCommands);
 	}
-	else if (Cast<UPlaneCutTool>(Tool) != nullptr)
+	else if (ExactCast<UEditMeshMaterialsTool>(Tool) != nullptr)
+	{
+		UPDATE_BINDING(FEditMeshMaterialsToolActionCommands);
+	}
+	else if (ExactCast<UPlaneCutTool>(Tool) != nullptr)
 	{
 		UPDATE_BINDING(FMeshPlaneCutToolActionCommands);
 	}
-	else if (Cast<UEditMeshPolygonsTool>(Tool) != nullptr)
+	else if (ExactCast<UEditMeshPolygonsTool>(Tool) != nullptr)
 	{
 		UPDATE_BINDING(FEditMeshPolygonsToolActionCommands);
 	}
-	else if (Cast<UDrawAndRevolveTool>(Tool) != nullptr)
+	else if (ExactCast<UDrawAndRevolveTool>(Tool) != nullptr)
 	{
 		UPDATE_BINDING(FDrawAndRevolveToolActionCommands);
 	}
@@ -175,6 +182,7 @@ DEFINE_TOOL_ACTION_COMMANDS(FMeshGroupPaintToolActionCommands, "ModelingToolsMes
 DEFINE_TOOL_ACTION_COMMANDS(FTransformToolActionCommands, "ModelingToolsTransformTool", "Modeling Tools - Transform Tool", UTransformMeshesTool);
 DEFINE_TOOL_ACTION_COMMANDS(FDrawPolygonToolActionCommands, "ModelingToolsDrawPolygonTool", "Modeling Tools - Draw Polygon Tool", UDrawPolygonTool);
 DEFINE_TOOL_ACTION_COMMANDS(FMeshSelectionToolActionCommands, "ModelingToolsMeshSelectionTool", "Modeling Tools - Mesh Selection Tool", UMeshSelectionTool);
+DEFINE_TOOL_ACTION_COMMANDS(FEditMeshMaterialsToolActionCommands, "ModelingToolsEditMeshMaterials", "Modeling Tools - Edit Materials Tool", UEditMeshMaterialsTool);
 DEFINE_TOOL_ACTION_COMMANDS(FMeshPlaneCutToolActionCommands, "ModelingToolsMeshPlaneCutTool", "Modeling Tools - Mesh Plane Cut Tool", UPlaneCutTool);
 DEFINE_TOOL_ACTION_COMMANDS(FEditMeshPolygonsToolActionCommands, "ModelingToolsEditMeshPolygonsTool", "Modeling Tools - Edit Mesh Polygons Tool", UEditMeshPolygonsTool);
 DEFINE_TOOL_ACTION_COMMANDS(FDrawAndRevolveToolActionCommands, "ModelingToolsDrawAndRevolveTool", "Modeling Tools - Draw-and-Revolve Tool", UDrawAndRevolveTool);
diff --git a/Engine/Plugins/Editor/ModelingToolsEditorMode/Source/ModelingToolsEditorMode/Public/ModelingToolsActions.h b/Engine/Plugins/Editor/ModelingToolsEditorMode/Source/ModelingToolsEditorMode/Public/ModelingToolsActions.h
--- a/Engine/Plugins/Editor/ModelingToolsEditorMode/Source/ModelingToolsEditorMode/Public/ModelingToolsActions.h
+++ b/Engine/Plugins/Editor/ModelingToolsEditorMode/Source/ModelingToolsEditorMode/Public/ModelingToolsActions.h
@@ -81,6 +81,7 @@ DECLARE_TOOL_ACTION_COMMANDS(FMeshGroupPaintToolActionCommands);
 DECLARE_TOOL_ACTION_COMMANDS(FTransformToolActionCommands);
 DECLARE_TOOL_ACTION_COMMANDS(FDrawPolygonToolActionCommands);
 DECLARE_TOOL_ACTION_COMMANDS(FMeshSelectionToolActionCommands);
+DECLARE_TOOL_ACTION_COMMANDS(FEditMeshMaterialsToolActionCommands);
 DECLARE_TOOL_ACTION_COMMANDS(FMeshPlaneCutToolActionCommands);
 DECLARE_TOOL_ACTION_COMMANDS(FEditMeshPolygonsToolActionCommands);
 DECLARE_TOOL_ACTION_COMMANDS(FDrawAndRevolveToolActionCommands);
diff --git a/Engine/Plugins/Experimental/MeshModelingToolsetExp/Source/MeshModelingToolsEditorOnly/Private/EditMeshMaterialsTool.cpp b/Engine/Plugins/Experimental/MeshModelingToolsetExp/Source/MeshModelingToolsEditorOnly/Private/EditMeshMaterialsTool.cpp
--- a/Engine/Plugins/Experimental/MeshModelingToolsetExp/Source/MeshModelingToolsEditorOnly/Private/EditMeshMaterialsTool.cpp
+++ b/Engine/Plugins/Experimental/MeshModelingToolsetExp/Source/MeshModelingToolsEditorOnly/Private/EditMeshMaterialsTool.cpp
@@ -160,6 +160,83 @@ void UEditMeshMaterialsTool::OnTick(float DeltaTime)
 }
 
 
+void UEditMeshMaterialsTool::RegisterActions(FInteractiveToolActionSet& ActionSet)
+{
+	UDynamicMeshBrushTool::RegisterActions(ActionSet);
+
+	// There's a bunch of code duplication here from UMeshSelectionTool::RegisterActions.
+	// In fact the only difference is currenlty that we don't register the "delete triangles" action.
+	// We could just override the function that performs the delete and not bother overriding the
+	// RegisterActions method at all, but that seems risky in case the select tool ever adds other
+	// things (especially lambdas) that we don't want to support in the edit materials tool...
+
+	ActionSet.RegisterAction(this, (int32)EStandardToolActions::BaseClientDefinedActionID + 50,
+		TEXT("TriSelectIncreaseSize"),
+		LOCTEXT("TriSelectIncreaseSize", "Increase Size"),
+		LOCTEXT("TriSelectIncreaseSizeTooltip", "Increase Brush Size"),
+		EModifierKey::None, EKeys::D,
+		[this]() { IncreaseBrushSizeAction(); });
+
+	ActionSet.RegisterAction(this, (int32)EStandardToolActions::BaseClientDefinedActionID + 51,
+		TEXT("TriSelectDecreaseSize"),
+		LOCTEXT("TriSelectDecreaseSize", "Decrease Size"),
+		LOCTEXT("TriSelectDecreaseSizeTooltip", "Decrease Brush Size"),
+		EModifierKey::None, EKeys::S,
+		[this]() { DecreaseBrushSizeAction(); });
+
+#if WITH_EDITOR  	// enum HasMetaData()  is not available at runtime
+	ActionSet.RegisterAction(this, (int32)EMeshSelectionToolActions::CycleSelectionMode,
+		TEXT("CycleSelectionMode"),
+		LOCTEXT("CycleSelectionMode", "Cycle Selection Mode"),
+		LOCTEXT("CycleSelectionModeTooltip", "Cycle through selection modes"),
+		EModifierKey::None, EKeys::Q,
+		[this]() {
+			const UEnum* SelectionModeEnum = StaticEnum<EMeshSelectionToolPrimaryMode>();
+			check(SelectionModeEnum);
+			int32 NumEnum = SelectionModeEnum->NumEnums() - 1;
+			do {
+				SelectionProps->SelectionMode = (EMeshSelectionToolPrimaryMode)(((int32)SelectionProps->SelectionMode + 1) % NumEnum);
+			} while (SelectionModeEnum->HasMetaData(TEXT("Hidden"), (int32)SelectionProps->SelectionMode));
+		}
+	);
+
+	ActionSet.RegisterAction(this, (int32)EMeshSelectionToolActions::CycleViewMode,
+		TEXT("CycleViewMode"),
+		LOCTEXT("CycleViewMode", "Cycle View Mode"),
+		LOCTEXT("CycleViewModeTooltip", "Cycle through face coloring modes"),
+		EModifierKey::None, EKeys::A,
+		[this]() {
+			const UEnum* ViewModeEnum = StaticEnum<EMeshFacesColorMode>();
+			check(ViewModeEnum);
+			int32 NumEnum = ViewModeEnum->NumEnums() - 1;
+			do {
+				SelectionProps->FaceColorMode = (EMeshFacesColorMode)(((int32)SelectionProps->FaceColorMode + 1) % NumEnum);
+			} while (ViewModeEnum->HasMetaData(TEXT("Hidden"), (int32)SelectionProps->FaceColorMode));
+		}
+	);
+#endif
+
+	ActionSet.RegisterAction(this, (int32)EMeshSelectionToolActions::ShrinkSelection,
+		TEXT("ShrinkSelection"),
+		LOCTEXT("ShrinkSelection", "Shrink Selection"),
+		LOCTEXT("ShrinkSelectionTooltip", "Shrink selection"),
+		EModifierKey::Shift, EKeys::Comma,
+		[this]() { GrowShrinkSelection(false); });
+
+	ActionSet.RegisterAction(this, (int32)EMeshSelectionToolActions::GrowSelection,
+		TEXT("GrowSelection"),
+		LOCTEXT("GrowSelection", "Grow Selection"),
+		LOCTEXT("GrowSelectionTooltip", "Grow selection"),
+		EModifierKey::Shift, EKeys::Period,
+		[this]() { GrowShrinkSelection(true); });
+
+	ActionSet.RegisterAction(this, (int32)EMeshSelectionToolActions::OptimizeSelection,
+		TEXT("OptimizeSelection"),
+		LOCTEXT("OptimizeSelection", "Optimize Selection"),
+		LOCTEXT("OptimizeSelectionTooltip", "Optimize selection"),
+		EModifierKey::None, EKeys::O,
+		[this]() { OptimizeSelection(); });
+}
 
 
 void UEditMeshMaterialsTool::ApplyMaterialAction(EEditMeshMaterialsToolActions ActionType)
diff --git a/Engine/Plugins/Experimental/MeshModelingToolsetExp/Source/MeshModelingToolsEditorOnly/Public/EditMeshMaterialsTool.h b/Engine/Plugins/Experimental/MeshModelingToolsetExp/Source/MeshModelingToolsEditorOnly/Public/EditMeshMaterialsTool.h
--- a/Engine/Plugins/Experimental/MeshModelingToolsetExp/Source/MeshModelingToolsEditorOnly/Public/EditMeshMaterialsTool.h
+++ b/Engine/Plugins/Experimental/MeshModelingToolsetExp/Source/MeshModelingToolsEditorOnly/Public/EditMeshMaterialsTool.h
@@ -88,6 +88,7 @@ class MESHMODELINGTOOLSEDITORONLY_API UEditMeshMaterialsTool : public UMeshSelec
 
 public:
 	virtual void Setup() override;
+	virtual void RegisterActions(FInteractiveToolActionSet& ActionSet) override;
 	virtual void OnTick(float DeltaTime) override;
 
 	virtual bool CanAccept() const override { return UMeshSelectionTool::CanAccept() || bHaveModifiedMaterials; }
diff --git a/Engine/Plugins/Experimental/MeshModelingToolsetExp/Source/MeshModelingToolsExp/Public/MeshSelectionTool.h b/Engine/Plugins/Experimental/MeshModelingToolsetExp/Source/MeshModelingToolsExp/Public/MeshSelectionTool.h
--- a/Engine/Plugins/Experimental/MeshModelingToolsetExp/Source/MeshModelingToolsExp/Public/MeshSelectionTool.h
+++ b/Engine/Plugins/Experimental/MeshModelingToolsetExp/Source/MeshModelingToolsExp/Public/MeshSelectionTool.h
@@ -125,14 +125,14 @@ public:
 
 	/** Select the largest connected mesh component by triangle count */
 	UFUNCTION(CallInEditor, Category = SelectionEdits, meta = (DisplayPriority = 6))
-	void LargestAreaPart()
+	void LargestTriCountPart()
 	{
 		PostAction(EMeshSelectionToolActions::SelectLargestComponentByTriCount);
 	}
 
 	/** Select the largest connected mesh component by mesh area */
 	UFUNCTION(CallInEditor, Category = SelectionEdits, meta = (DisplayPriority = 7))
-	void LargestTriCountPart()
+	void LargestAreaPart()
 	{
 		PostAction(EMeshSelectionToolActions::SelectLargestComponentByArea);
 	}
