ModelingMode: make persistent selections an optional, disabled-by-default feature.
#rnx

diff --git a/Engine/Plugins/Editor/ModelingToolsEditorMode/Source/ModelingToolsEditorMode/Private/ModelingToolsEditorMode.cpp b/Engine/Plugins/Editor/ModelingToolsEditorMode/Source/ModelingToolsEditorMode/Private/ModelingToolsEditorMode.cpp
--- a/Engine/Plugins/Editor/ModelingToolsEditorMode/Source/ModelingToolsEditorMode/Private/ModelingToolsEditorMode.cpp
+++ b/Engine/Plugins/Editor/ModelingToolsEditorMode/Source/ModelingToolsEditorMode/Private/ModelingToolsEditorMode.cpp
@@ -3,6 +3,7 @@
 #include "ModelingToolsEditorMode.h"
 #include "InteractiveTool.h"
 #include "ModelingToolsEditorModeToolkit.h"
+#include "ModelingToolsEditorModeSettings.h"
 #include "Toolkits/ToolkitManager.h"
 #include "ToolTargetManager.h"
 #include "ContextObjectStore.h"
@@ -317,8 +318,13 @@ void UModelingToolsEditorMode::Enter()
 
 	// register gizmo helper
 	UE::TransformGizmoUtil::RegisterTransformGizmoContextObject(GetInteractiveToolsContext());
-	// register selection manager
-	UE::Geometry::RegisterPersistentMeshSelectionManager(GetInteractiveToolsContext());
+
+	// register selection manager, if this feature is enabled in the mode settings
+	const UModelingToolsEditorModeSettings* ModelingModeSettings = GetDefault<UModelingToolsEditorModeSettings>();
+	if (ModelingModeSettings && ModelingModeSettings->bEnablePersistentSelections)
+	{
+		UE::Geometry::RegisterPersistentMeshSelectionManager(GetInteractiveToolsContext());
+	}
 
 	// disable HitProxy rendering, it is not used in Modeling Mode and adds overhead to Render() calls
 	GetInteractiveToolsContext()->SetEnableRenderingDuringHitProxyPass(false);
@@ -707,24 +713,28 @@ void UModelingToolsEditorMode::Enter()
 	// Unfortunately most routes to find out about selection changes don't allow for this. Currently this
 	// OnPreChange event is the only one that appears to provide the desired behavior, however it is likely
 	// that this is not going to be reliable...
-	SelectionModifiedEventHandle = GetModeManager()->GetSelectedActors()->GetElementSelectionSet()->OnPreChange().AddLambda( [&](const UTypedElementSelectionSet*) 
-	{ 
-		if (GIsTransacting == 0)
+	if (ModelingModeSettings && ModelingModeSettings->bEnablePersistentSelections)
+	{
+		SelectionModifiedEventHandle = GetModeManager()->GetSelectedActors()->GetElementSelectionSet()->OnPreChange().AddLambda([&](const UTypedElementSelectionSet*)
 		{
-			UE::Geometry::ClearActiveToolSelection(GetToolManager());
-		}
-	});
+			if (GIsTransacting == 0)
+			{
+				UE::Geometry::ClearActiveToolSelection(GetToolManager());
+			}
+		});
+	}
 }
 
 void UModelingToolsEditorMode::Exit()
 {
 	// clear any active selection
 	UE::Geometry::ClearActiveToolSelection(GetToolManager());
-	// deregister selection manager
+
+	// deregister selection manager (note: may not have been registered, depending on mode settings)
 	UE::Geometry::DeregisterPersistentMeshSelectionManager(GetInteractiveToolsContext());
 
 	// stop listening to selection changes. On Editor Shutdown, some of these values become null, which will result in an ensure/crash
-	if (UObjectInitialized() && GetModeManager() && GetModeManager()->GetSelectedActors() && GetModeManager()->GetSelectedActors()->GetElementSelectionSet() )
+	if (SelectionModifiedEventHandle.IsValid() && UObjectInitialized() && GetModeManager() && GetModeManager()->GetSelectedActors() && GetModeManager()->GetSelectedActors()->GetElementSelectionSet() )
 	{
 		GetModeManager()->GetSelectedActors()->GetElementSelectionSet()->OnPreChange().Remove(SelectionModifiedEventHandle);
 	}
diff --git a/Engine/Plugins/Editor/ModelingToolsEditorMode/Source/ModelingToolsEditorMode/Public/ModelingToolsEditorModeSettings.h b/Engine/Plugins/Editor/ModelingToolsEditorMode/Source/ModelingToolsEditorMode/Public/ModelingToolsEditorModeSettings.h
--- a/Engine/Plugins/Editor/ModelingToolsEditorMode/Source/ModelingToolsEditorMode/Public/ModelingToolsEditorModeSettings.h
+++ b/Engine/Plugins/Editor/ModelingToolsEditorMode/Source/ModelingToolsEditorMode/Public/ModelingToolsEditorModeSettings.h
@@ -92,6 +92,16 @@ public:
 	UPROPERTY(config, EditAnywhere, Category = "Modeling Mode|Generated Assets")
 	bool bAppendRandomStringToName = true;
 
+
+	//
+	// Selection
+	//
+
+
+	/** If true, selections will be tracked between tools, and visualized with outlines. This feature is Experimental, disabled by default. */
+	UPROPERTY(config, EditAnywhere, Category = "Modeling Mode|Selection")
+	bool bEnablePersistentSelections = false;
+
 public:
 
 	DECLARE_MULTICAST_DELEGATE_TwoParams(UModelingToolsEditorModeSettingsModified, UObject*, FProperty*);
