MeshModelingToolset: Fix data collection for tools usage using localized names

#rnx

diff --git a/Engine/Plugins/Editor/ModelingToolsEditorMode/Source/ModelingToolsEditorMode/Private/ModelingToolsEditorMode.cpp b/Engine/Plugins/Editor/ModelingToolsEditorMode/Source/ModelingToolsEditorMode/Private/ModelingToolsEditorMode.cpp
--- a/Engine/Plugins/Editor/ModelingToolsEditorMode/Source/ModelingToolsEditorMode/Private/ModelingToolsEditorMode.cpp
+++ b/Engine/Plugins/Editor/ModelingToolsEditorMode/Source/ModelingToolsEditorMode/Private/ModelingToolsEditorMode.cpp
@@ -125,6 +125,15 @@ const FEditorModeID UModelingToolsEditorMode::EM_ModelingToolsEditorModeId = TEX
 FDateTime UModelingToolsEditorMode::LastModeStartTimestamp;
 FDateTime UModelingToolsEditorMode::LastToolStartTimestamp;
 
+namespace
+{
+FString GetToolName(const UInteractiveTool& Tool)
+{
+	const FString* ToolName = FTextInspector::GetSourceString(Tool.GetToolInfo().ToolDisplayName);
+	return ToolName ? *ToolName : FString(TEXT("<Invalid ToolName>"));
+}
+}
+
 UModelingToolsEditorMode::UModelingToolsEditorMode()
 {
 	Info = FEditorModeInfo(
@@ -651,9 +660,11 @@ void UModelingToolsEditorMode::Enter()
 	// Log mode starting
 	if (FEngineAnalytics::IsAvailable())
 	{
-		UModelingToolsEditorMode::LastModeStartTimestamp = FDateTime::UtcNow();
+		LastModeStartTimestamp = FDateTime::UtcNow();
+
 		TArray<FAnalyticsEventAttribute> Attributes;
 		Attributes.Add(FAnalyticsEventAttribute(TEXT("Timestamp"), LastModeStartTimestamp.ToString()));
+
 		FEngineAnalytics::GetProvider().RecordEvent(TEXT("Editor.Usage.MeshModelingMode.Enter"), Attributes);
 	}
 
@@ -662,10 +673,12 @@ void UModelingToolsEditorMode::Enter()
 	{
 		if (FEngineAnalytics::IsAvailable())
 		{
-			UModelingToolsEditorMode::LastToolStartTimestamp = FDateTime::UtcNow();
+			LastToolStartTimestamp = FDateTime::UtcNow();
+
 			TArray<FAnalyticsEventAttribute> Attributes;
-			Attributes.Add(FAnalyticsEventAttribute(TEXT("DisplayName"), Tool->GetToolInfo().ToolDisplayName.ToString()));
+			Attributes.Add(FAnalyticsEventAttribute(TEXT("ToolName"), GetToolName(*Tool)));
 			Attributes.Add(FAnalyticsEventAttribute(TEXT("Timestamp"), LastToolStartTimestamp.ToString()));
+
 			FEngineAnalytics::GetProvider().RecordEvent(TEXT("Editor.Usage.MeshModelingMode.ToolStarted"), Attributes);
 		}
 	});
@@ -675,11 +688,14 @@ void UModelingToolsEditorMode::Enter()
 	{
 		if (FEngineAnalytics::IsAvailable())
 		{
+			const FDateTime Now = FDateTime::UtcNow();
+			const FTimespan ToolUsageDuration = Now - LastToolStartTimestamp;
+
 			TArray<FAnalyticsEventAttribute> Attributes;
-			Attributes.Add(FAnalyticsEventAttribute(TEXT("DisplayName"), Tool->GetToolInfo().ToolDisplayName.ToString()));
-			Attributes.Add(FAnalyticsEventAttribute(TEXT("Timestamp"), FDateTime::UtcNow().ToString()));
-			const FTimespan ToolUsageDuration = FDateTime::UtcNow() - UModelingToolsEditorMode::LastToolStartTimestamp;
+			Attributes.Add(FAnalyticsEventAttribute(TEXT("ToolName"), GetToolName(*Tool)));
+			Attributes.Add(FAnalyticsEventAttribute(TEXT("Timestamp"), Now.ToString()));
 			Attributes.Add(FAnalyticsEventAttribute(TEXT("Duration.Seconds"), static_cast<float>(ToolUsageDuration.GetTotalSeconds())));
+
 			FEngineAnalytics::GetProvider().RecordEvent(TEXT("Editor.Usage.MeshModelingMode.ToolEnded"), Attributes);
 		}
 	});
@@ -728,10 +744,12 @@ void UModelingToolsEditorMode::Exit()
 	// Log mode ending
 	if (FEngineAnalytics::IsAvailable())
 	{
+		const FTimespan ModeUsageDuration = FDateTime::UtcNow() - LastModeStartTimestamp;
+
 		TArray<FAnalyticsEventAttribute> Attributes;
 		Attributes.Add(FAnalyticsEventAttribute(TEXT("Timestamp"), FDateTime::UtcNow().ToString()));
-		const FTimespan ModeUsageDuration = FDateTime::UtcNow() - UModelingToolsEditorMode::LastModeStartTimestamp;
 		Attributes.Add(FAnalyticsEventAttribute(TEXT("Duration.Seconds"), static_cast<float>(ModeUsageDuration.GetTotalSeconds())));
+
 		FEngineAnalytics::GetProvider().RecordEvent(TEXT("Editor.Usage.MeshModelingMode.Exit"), Attributes);
 	}
 
@@ -775,9 +793,9 @@ void UModelingToolsEditorMode::OnToolStarted(UInteractiveToolManager* Manager, U
 	
 	if( FEngineAnalytics::IsAvailable() )
 	{
-		FEngineAnalytics::GetProvider().RecordEvent( TEXT( "Editor.Usage.MeshModelingMode.ToolStarted" ),
-													 TEXT( "DisplayName" ),
-													 Tool->GetToolInfo().ToolDisplayName.ToString() );
+		FEngineAnalytics::GetProvider().RecordEvent(TEXT("Editor.Usage.MeshModelingMode.ToolStarted"),
+		                                            TEXT("ToolName"),
+		                                            GetToolName(*Tool));
 	}
 }
 
@@ -787,9 +805,9 @@ void UModelingToolsEditorMode::OnToolEnded(UInteractiveToolManager* Manager, UIn
 	
 	if( FEngineAnalytics::IsAvailable() )
 	{
-		FEngineAnalytics::GetProvider().RecordEvent( TEXT( "Editor.Usage.MeshModelingMode.ToolEnded" ),
-													 TEXT( "DisplayName" ),
-													 Tool->GetToolInfo().ToolDisplayName.ToString() );
+		FEngineAnalytics::GetProvider().RecordEvent(TEXT("Editor.Usage.MeshModelingMode.ToolEnded"),
+		                                            TEXT("ToolName"),
+		                                            GetToolName(*Tool));
 	}
 }
 
