ModelingMode: change settings object to be a UDeveloperSettings, remove manual registration. Fix settings quick-access icon.
#rnx

diff --git a/Engine/Plugins/Editor/ModelingToolsEditorMode/Source/ModelingToolsEditorMode/ModelingToolsEditorMode.Build.cs b/Engine/Plugins/Editor/ModelingToolsEditorMode/Source/ModelingToolsEditorMode/ModelingToolsEditorMode.Build.cs
--- a/Engine/Plugins/Editor/ModelingToolsEditorMode/Source/ModelingToolsEditorMode/ModelingToolsEditorMode.Build.cs
+++ b/Engine/Plugins/Editor/ModelingToolsEditorMode/Source/ModelingToolsEditorMode/ModelingToolsEditorMode.Build.cs
@@ -63,14 +63,13 @@ public class ModelingToolsEditorMode : ModuleRules
 				"MeshModelingToolsEditorOnly",
 				"MeshLODToolset",
 				"ToolWidgets",
-
-                "ViewportInteraction",
-
-				"StylusInput"
+				"EditorWidgets",
+				"StylusInput",
+				"DeveloperSettings"
 
 				// ... add private dependencies that you statically link with here ...	
 			}
-            );
+			);
 		
 		
 		DynamicallyLoadedModuleNames.AddRange(
diff --git a/Engine/Plugins/Editor/ModelingToolsEditorMode/Source/ModelingToolsEditorMode/Private/ModelingToolsEditorModeModule.cpp b/Engine/Plugins/Editor/ModelingToolsEditorMode/Source/ModelingToolsEditorMode/Private/ModelingToolsEditorModeModule.cpp
--- a/Engine/Plugins/Editor/ModelingToolsEditorMode/Source/ModelingToolsEditorMode/Private/ModelingToolsEditorModeModule.cpp
+++ b/Engine/Plugins/Editor/ModelingToolsEditorMode/Source/ModelingToolsEditorMode/Private/ModelingToolsEditorModeModule.cpp
@@ -7,8 +7,6 @@
 #include "ModelingToolsEditorModeStyle.h"
 #include "ModelingToolsEditorModeSettings.h"
 
-#include "ISettingsModule.h"
-#include "ISettingsSection.h"
 #include "Misc/CoreDelegates.h"
 
 #include "Modules/ModuleManager.h"
@@ -26,11 +24,6 @@ void FModelingToolsEditorModeModule::ShutdownModule()
 {
 	FCoreDelegates::OnPostEngineInit.RemoveAll(this);
 
-	if (ISettingsModule* SettingsModule = FModuleManager::GetModulePtr<ISettingsModule>("Settings"))
-	{
-		SettingsModule->UnregisterSettings("Project", "Plugins", "ModelingMode");
-	}
-
 	FModelingToolActionCommands::UnregisterAllToolActions();
 	FModelingToolsManagerCommands::Unregister();
 	FModelingModeActionCommands::Unregister();
@@ -49,17 +42,6 @@ void FModelingToolsEditorModeModule::OnPostEngineInit()
 	FModelingToolsManagerCommands::Register();
 	FModelingModeActionCommands::Register();
 
-	// register settings
-	ISettingsModule* SettingsModule = FModuleManager::GetModulePtr<ISettingsModule>("Settings");
-	if (SettingsModule != nullptr)
-	{
-		ISettingsSectionPtr SettingsSection = SettingsModule->RegisterSettings("Project", "Plugins", "ModelingMode",
-			LOCTEXT("ModelingModeSettingsName", "Modeling Mode"),
-			LOCTEXT("ModelingModeSettingsDescription", "Configure the Modeling Tools Editor Mode plugin"),
-			GetMutableDefault<UModelingToolsEditorModeSettings>()
-		);
-	}
-
 	// Register details view customizations
 	FPropertyEditorModule& PropertyModule = FModuleManager::LoadModuleChecked<FPropertyEditorModule>("PropertyEditor");
 	PropertyModule.RegisterCustomPropertyTypeLayout("ModelingToolsAxisFilter", FOnGetPropertyTypeCustomizationInstance::CreateStatic(&FModelingToolsAxisFilterCustomization::MakeInstance));
diff --git a/Engine/Plugins/Editor/ModelingToolsEditorMode/Source/ModelingToolsEditorMode/Private/ModelingToolsEditorModeSettings.cpp b/Engine/Plugins/Editor/ModelingToolsEditorMode/Source/ModelingToolsEditorMode/Private/ModelingToolsEditorModeSettings.cpp
new file mode 100644
--- /dev/null
+++ b/Engine/Plugins/Editor/ModelingToolsEditorMode/Source/ModelingToolsEditorMode/Private/ModelingToolsEditorModeSettings.cpp
@@ -0,0 +1,19 @@
+// Copyright Epic Games, Inc. All Rights Reserved.
+
+#include "ModelingToolsEditorModeSettings.h"
+
+
+#define LOCTEXT_NAMESPACE "ModelingToolsEditorModeSettings"
+
+FText UModelingToolsEditorModeSettings::GetSectionText() const 
+{ 
+	return LOCTEXT("ModelingModeSettingsName", "Modeling Mode"); 
+}
+
+FText UModelingToolsEditorModeSettings::GetSectionDescription() const
+{
+	return LOCTEXT("ModelingModeSettingsDescription", "Configure the Modeling Tools Editor Mode plugin");
+}
+
+
+#undef LOCTEXT_NAMESPACE
\ No newline at end of file
diff --git a/Engine/Plugins/Editor/ModelingToolsEditorMode/Source/ModelingToolsEditorMode/Private/ModelingToolsEditorModeToolkit.cpp b/Engine/Plugins/Editor/ModelingToolsEditorMode/Source/ModelingToolsEditorMode/Private/ModelingToolsEditorModeToolkit.cpp
--- a/Engine/Plugins/Editor/ModelingToolsEditorMode/Source/ModelingToolsEditorMode/Private/ModelingToolsEditorModeToolkit.cpp
+++ b/Engine/Plugins/Editor/ModelingToolsEditorMode/Source/ModelingToolsEditorMode/Private/ModelingToolsEditorModeToolkit.cpp
@@ -17,6 +17,7 @@
 #include "Toolkits/AssetEditorModeUILayer.h"
 
 #include "Widgets/Input/SButton.h"
+#include "SSimpleButton.h"
 #include "Widgets/Text/STextBlock.h"
 #include "Widgets/Layout/SSeparator.h"
 #include "Widgets/Images/SImage.h"
@@ -304,7 +305,7 @@ TSharedPtr<SWidget> FModelingToolsEditorModeToolkit::MakeAssetConfigPanel()
 	AssetLODModeLabel = SNew(STextBlock).Text(LOCTEXT("ActiveLODLabel", "Editing LOD"));
 
 	TSharedPtr<SVerticalBox> Content = SNew(SVerticalBox)
-	+ SVerticalBox::Slot().MaxHeight(20).HAlign(HAlign_Fill)
+	+ SVerticalBox::Slot().HAlign(HAlign_Fill)
 		[
 			SNew(SHorizontalBox)
 			+ SHorizontalBox::Slot().Padding(0, 2, 2, 2).HAlign(HAlign_Right).VAlign(VAlign_Center).AutoWidth()
@@ -313,7 +314,7 @@ TSharedPtr<SWidget> FModelingToolsEditorModeToolkit::MakeAssetConfigPanel()
 				[ AssetLODMode->AsShared() ]
 		]
 
-	+ SVerticalBox::Slot().MaxHeight(20).HAlign(HAlign_Fill).Padding(0, 3, 0, 0)
+	+ SVerticalBox::Slot().HAlign(HAlign_Fill).Padding(0, 3, 0, 0)
 		[
 			SNew(SHorizontalBox)
 			+ SHorizontalBox::Slot().Padding(0, 2, 2, 2).HAlign(HAlign_Right).VAlign(VAlign_Center).AutoWidth()
@@ -321,23 +322,15 @@ TSharedPtr<SWidget> FModelingToolsEditorModeToolkit::MakeAssetConfigPanel()
 					SNew(STextBlock)
 					.Text(LOCTEXT("AssetLocationLabel", "New Asset Location"))
 				]
-			+ SHorizontalBox::Slot().Padding(0).FillWidth(4.0f)
+			+ SHorizontalBox::Slot().HAlign(HAlign_Fill).Padding(0).FillWidth(1.0f)
 				[
 					AssetLocationMode->AsShared()
 				]
-
-			+ SHorizontalBox::Slot().HAlign(HAlign_Right).Padding(5, 0, 0, 0).FillWidth(1.0f)
+			+ SHorizontalBox::Slot().HAlign(HAlign_Right).Padding(0, 0, 0, 0).AutoWidth()
 				[
-					SNew(SBox).MaxDesiredHeight(16).MaxDesiredWidth(17)
-					[
-						SNew(SButton)
-						.ContentPadding(0)
-						.ButtonStyle(FEditorStyle::Get(), "FlatButton.Default")
-						.OnClicked_Lambda( [this]() {OnShowAssetSettings(); return FReply::Handled(); } )
-						[
-							SNew(SImage).Image(FSlateIcon(FEditorStyle::GetStyleSetName(), "FoliageEditMode.Settings").GetIcon())
-						]
-					]
+					SNew(SSimpleButton)
+					.OnClicked_Lambda( [this]() { OnShowAssetSettings(); return FReply::Handled(); } )
+					.Icon(FAppStyle::Get().GetBrush("Icons.Settings"))
 				]
 		];
 
diff --git a/Engine/Plugins/Editor/ModelingToolsEditorMode/Source/ModelingToolsEditorMode/Public/ModelingToolsEditorModeSettings.h b/Engine/Plugins/Editor/ModelingToolsEditorMode/Source/ModelingToolsEditorMode/Public/ModelingToolsEditorModeSettings.h
--- a/Engine/Plugins/Editor/ModelingToolsEditorMode/Source/ModelingToolsEditorMode/Public/ModelingToolsEditorModeSettings.h
+++ b/Engine/Plugins/Editor/ModelingToolsEditorMode/Source/ModelingToolsEditorMode/Public/ModelingToolsEditorModeSettings.h
@@ -2,8 +2,7 @@
 
 #pragma once
 
-#include "UObject/ObjectMacros.h"
-#include "UObject/Object.h"
+#include "Engine/DeveloperSettings.h"
 #include "ModelingToolsEditorModeSettings.generated.h"
 
 
@@ -54,10 +53,21 @@ enum class EModelingModeAssetGenerationLocation
  * Settings for the Modeling Tools Editor Mode plug-in.
  */
 UCLASS(config=Editor)
-class MODELINGTOOLSEDITORMODE_API UModelingToolsEditorModeSettings
-	: public UObject
+class MODELINGTOOLSEDITORMODE_API UModelingToolsEditorModeSettings : public UDeveloperSettings
 {
 	GENERATED_BODY()
+
+public:
+
+	// UDeveloperSettings overrides
+
+	virtual FName GetContainerName() const { return FName("Project"); }
+	virtual FName GetCategoryName() const { return FName("Plugins"); }
+	virtual FName GetSectionName() const { return FName("ModelingMode"); }
+
+	virtual FText GetSectionText() const override;
+	virtual FText GetSectionDescription() const override;
+
 public:
 
 	/** Enable/Disable the options to emit Dynamic Mesh Actors in Modeling Mode Tools */
@@ -107,9 +117,11 @@ public:
 	DECLARE_MULTICAST_DELEGATE_TwoParams(UModelingToolsEditorModeSettingsModified, UObject*, FProperty*);
 	UModelingToolsEditorModeSettingsModified OnModified;
 
-	void PostEditChangeProperty(FPropertyChangedEvent& PropertyChangedEvent)
+	virtual void PostEditChangeProperty(FPropertyChangedEvent& PropertyChangedEvent) override
 	{
 		OnModified.Broadcast(this, PropertyChangedEvent.Property);
+
+		Super::PostEditChangeProperty(PropertyChangedEvent);
 	}
 
 };
