ModelingMode: Add user name override for auto-generated asset folder

diff --git a/Engine/Plugins/Editor/ModelingToolsEditorMode/Source/ModelingToolsEditorMode/Private/ModelingModeAssetUtils.cpp b/Engine/Plugins/Editor/ModelingToolsEditorMode/Source/ModelingToolsEditorMode/Private/ModelingModeAssetUtils.cpp
--- a/Engine/Plugins/Editor/ModelingToolsEditorMode/Source/ModelingToolsEditorMode/Private/ModelingModeAssetUtils.cpp
+++ b/Engine/Plugins/Editor/ModelingToolsEditorMode/Source/ModelingToolsEditorMode/Private/ModelingModeAssetUtils.cpp
@@ -161,8 +161,12 @@ FString UE::Modeling::GetNewAssetPathName(const FString& BaseNameIn, const UWorl
 		// append username-specific subfolder
 		if (Settings->bUsePerUserAutogenSubfolder)
 		{
-			FString UsernameString(FPlatformProcess::UserName());
-			if (UsernameString.Len() > 1)
+			FString UsernameString = Settings->AutogenSubfolderUserNameOverride.TrimStartAndEnd();
+			if (UsernameString.IsEmpty())
+			{
+				UsernameString = FPlatformProcess::UserName();
+			}
+			if (!UsernameString.IsEmpty())
 			{
 				PackageFolderPath = FPaths::Combine(PackageFolderPath, UsernameString);
 			}
diff --git a/Engine/Plugins/Editor/ModelingToolsEditorMode/Source/ModelingToolsEditorMode/Public/ModelingToolsEditorModeSettings.h b/Engine/Plugins/Editor/ModelingToolsEditorMode/Source/ModelingToolsEditorMode/Public/ModelingToolsEditorModeSettings.h
--- a/Engine/Plugins/Editor/ModelingToolsEditorMode/Source/ModelingToolsEditorMode/Public/ModelingToolsEditorModeSettings.h
+++ b/Engine/Plugins/Editor/ModelingToolsEditorMode/Source/ModelingToolsEditorMode/Public/ModelingToolsEditorModeSettings.h
@@ -98,6 +98,10 @@ public:
 	UPROPERTY(config, EditAnywhere, Category = "Modeling Mode|Generated Assets")
 	bool bUsePerUserAutogenSubfolder = true;
 
+	/** Overrides the user name used for per-user folders below the Autogen path. This might be necessary to resolve issues with source control, for example. Note that the per-user folder name might not contain the name exactly as provided. */
+	UPROPERTY(config, EditAnywhere, Category = "Modeling Mode|Generated Assets")
+	FString AutogenSubfolderUserNameOverride;
+
 	/** If true, Autogenerated Assets have a short random string generated and appended to their name */
 	UPROPERTY(config, EditAnywhere, Category = "Modeling Mode|Generated Assets")
 	bool bAppendRandomStringToName = true;
