ModelingMode: support for tool palette customization
- Add UModelingToolsModeCustomizationSettings for specifying Editor-level Modeling Mode UI Customization settings
- add support for re-ordering Modeling Mode Tool Palettes based on list in customization settings
- add support for adding specific Tools to a Favorites section in the Palette, based on list in customization settings
- add support for setting custom header colors for sections in in tool palette, based on list in customization settings
- move Baking tools to a separate Baking section
#rnx

diff --git a/Engine/Plugins/Editor/ModelingToolsEditorMode/Source/ModelingToolsEditorMode/Private/ModelingToolsEditorModeSettings.cpp b/Engine/Plugins/Editor/ModelingToolsEditorMode/Source/ModelingToolsEditorMode/Private/ModelingToolsEditorModeSettings.cpp
--- a/Engine/Plugins/Editor/ModelingToolsEditorMode/Source/ModelingToolsEditorMode/Private/ModelingToolsEditorModeSettings.cpp
+++ b/Engine/Plugins/Editor/ModelingToolsEditorMode/Source/ModelingToolsEditorMode/Private/ModelingToolsEditorModeSettings.cpp
@@ -5,6 +5,7 @@
 
 #define LOCTEXT_NAMESPACE "ModelingToolsEditorModeSettings"
 
+
 FText UModelingToolsEditorModeSettings::GetSectionText() const 
 { 
 	return LOCTEXT("ModelingModeSettingsName", "Modeling Mode"); 
@@ -16,4 +17,15 @@ FText UModelingToolsEditorModeSettings::GetSectionDescription() const
 }
 
 
+FText UModelingToolsModeCustomizationSettings::GetSectionText() const 
+{ 
+	return LOCTEXT("ModelingModeSettingsName", "Modeling Mode"); 
+}
+
+FText UModelingToolsModeCustomizationSettings::GetSectionDescription() const
+{
+	return LOCTEXT("ModelingModeSettingsDescription", "Configure the Modeling Tools Editor Mode plugin");
+}
+
+
 #undef LOCTEXT_NAMESPACE
\ No newline at end of file
diff --git a/Engine/Plugins/Editor/ModelingToolsEditorMode/Source/ModelingToolsEditorMode/Private/ModelingToolsEditorModeToolkit.cpp b/Engine/Plugins/Editor/ModelingToolsEditorMode/Source/ModelingToolsEditorMode/Private/ModelingToolsEditorModeToolkit.cpp
--- a/Engine/Plugins/Editor/ModelingToolsEditorMode/Source/ModelingToolsEditorMode/Private/ModelingToolsEditorModeToolkit.cpp
+++ b/Engine/Plugins/Editor/ModelingToolsEditorMode/Source/ModelingToolsEditorMode/Private/ModelingToolsEditorModeToolkit.cpp
@@ -453,9 +453,11 @@ static const FName PrototypesTabName(TEXT("Prototypes"));
 static const FName PolyEditTabName(TEXT("PolyEdit"));
 static const FName VoxToolsTabName(TEXT("VoxOps"));
 static const FName LODToolsTabName(TEXT("LODs"));
+static const FName BakingToolsTabName(TEXT("Baking"));
+static const FName ModelingFavoritesTabName(TEXT("Favorites"));
 
 
-const TArray<FName> FModelingToolsEditorModeToolkit::PaletteNames_Standard = { PrimitiveTabName, CreateTabName, PolyModelingTabName, TriModelingTabName, DeformTabName, TransformTabName, MeshProcessingTabName, VoxToolsTabName, AttributesTabName, UVTabName, VolumesTabName, LODToolsTabName };
+const TArray<FName> FModelingToolsEditorModeToolkit::PaletteNames_Standard = { PrimitiveTabName, CreateTabName, PolyModelingTabName, TriModelingTabName, DeformTabName, TransformTabName, MeshProcessingTabName, VoxToolsTabName, AttributesTabName, UVTabName, BakingToolsTabName, VolumesTabName, LODToolsTabName };
 
 
 void FModelingToolsEditorModeToolkit::GetToolPaletteNames(TArray<FName>& PaletteNames) const
@@ -503,6 +505,35 @@ void FModelingToolsEditorModeToolkit::GetToolPaletteNames(TArray<FName>& Palette
 			}
 		}
 	}
+
+	UModelingToolsModeCustomizationSettings* UISettings = GetMutableDefault<UModelingToolsModeCustomizationSettings>();
+
+	// if user has provided custom ordering of tool palettes in the Editor Settings, try to apply them
+	if (UISettings->ToolSectionOrder.Num() > 0)
+	{
+		TArray<FName> NewPaletteNames;
+		for (FString SectionName : UISettings->ToolSectionOrder)
+		{
+			for (int32 k = 0; k < PaletteNames.Num(); ++k)
+			{
+				if (SectionName.Compare(PaletteNames[k].ToString(), ESearchCase::IgnoreCase) == 0)
+				{
+					NewPaletteNames.Add(PaletteNames[k]);
+					PaletteNames.RemoveAt(k);
+					break;
+				}
+			}
+		}
+		NewPaletteNames.Append(PaletteNames);
+		PaletteNames = MoveTemp(NewPaletteNames);
+	}
+
+	// if user has provided a list of favorite tools, add that palette to the list
+	if (UISettings->ToolFavorites.Num() > 0)
+	{
+		PaletteNames.Insert(ModelingFavoritesTabName, 0);
+	}
+
 }
 
 
@@ -514,21 +545,25 @@ FText FModelingToolsEditorModeToolkit::GetToolPaletteDisplayName(FName Palette)
 void FModelingToolsEditorModeToolkit::BuildToolPalette(FName PaletteIndex, class FToolBarBuilder& ToolbarBuilder) 
 {
 	const FModelingToolsManagerCommands& Commands = FModelingToolsManagerCommands::Get();
-	BuildToolPalette_Experimental(PaletteIndex, ToolbarBuilder);
-}
-
-
-
-void FModelingToolsEditorModeToolkit::BuildToolPalette_Standard(FName PaletteIndex, class FToolBarBuilder& ToolbarBuilder)
-{
-	check(false);
-}
-
-
-void FModelingToolsEditorModeToolkit::BuildToolPalette_Experimental(FName PaletteIndex, class FToolBarBuilder& ToolbarBuilder)
-{
-	const FModelingToolsManagerCommands& Commands = FModelingToolsManagerCommands::Get();
+	UModelingToolsModeCustomizationSettings* UISettings = GetMutableDefault<UModelingToolsModeCustomizationSettings>();
 
+	if (PaletteIndex == ModelingFavoritesTabName)
+	{
+		// build Favorites tool palette
+		for (FString ToolName : UISettings->ToolFavorites)
+		{
+			bool bFound = false;
+			TSharedPtr<FUICommandInfo> FoundToolCommand = Commands.FindToolByName(ToolName, bFound);
+			if ( bFound )
+			{ 
+				ToolbarBuilder.AddToolBarButton(FoundToolCommand);
+			}
+			else
+			{
+				UE_LOG(LogTemp, Display, TEXT("ModelingMode: could not find Favorited Tool %s"), *ToolName);
+			}
+		}
+	}
 	if (PaletteIndex == PrimitiveTabName)
 	{
 		ToolbarBuilder.AddToolBarButton(Commands.BeginAddBoxPrimitiveTool);
@@ -622,6 +657,9 @@ void FModelingToolsEditorModeToolkit::BuildToolPalette_Experimental(FName Palett
 		ToolbarBuilder.AddToolBarButton(Commands.BeginMeshGroupPaintTool);
 		ToolbarBuilder.AddToolBarButton(Commands.BeginMeshAttributePaintTool);
 		ToolbarBuilder.AddToolBarButton(Commands.BeginEditMeshMaterialsTool);
+	} 
+	else if (PaletteIndex == BakingToolsTabName )
+	{
 		ToolbarBuilder.AddToolBarButton(Commands.BeginBakeMeshAttributeMapsTool);
 		ToolbarBuilder.AddToolBarButton(Commands.BeginBakeMultiMeshAttributeMapsTool);
 		ToolbarBuilder.AddToolBarButton(Commands.BeginBakeMeshAttributeVertexTool);
@@ -692,6 +730,64 @@ void FModelingToolsEditorModeToolkit::BuildToolPalette_Experimental(FName Palett
 	}
 }
 
+
+void FModelingToolsEditorModeToolkit::InvokeUI()
+{
+	FModeToolkit::InvokeUI();
+
+	//
+	// Apply custom section header colors.
+	// See comments below, this is done via directly manipulating Slate widgets generated deep inside BaseToolkit.cpp,
+	// and will stop working if the Slate widget structure changes
+	//
+
+	UModelingToolsModeCustomizationSettings* UISettings = GetMutableDefault<UModelingToolsModeCustomizationSettings>();
+
+	// look up default radii for palette toolbar expandable area headers
+	FVector4 HeaderRadii(4, 4, 0, 0);
+	const FSlateBrush* BaseBrush = FAppStyle::Get().GetBrush("PaletteToolbar.ExpandableAreaHeader");
+	if (BaseBrush != nullptr)
+	{
+		HeaderRadii = BaseBrush->OutlineSettings.CornerRadii;
+	}
+
+	for (FEdModeToolbarRow& ToolbarRow : ActiveToolBarRows)
+	{
+		for (FModelingModeCustomSectionColor ToolColor : UISettings->SectionColors)
+		{
+			if (ToolColor.SectionName.Compare(ToolbarRow.PaletteName.ToString(), ESearchCase::IgnoreCase) == 0)
+			{
+				// code below is highly dependent on the structure of the ToolbarRow.ToolbarWidget. Currently this is 
+				// a SMultiBoxWidget, a few levels below a SExpandableArea. The SExpandableArea contains a SVerticalBox
+				// with the header as a SBorder in Slot 0. The code will fail gracefully if this structure changes.
+
+				TSharedPtr<SWidget> ExpanderVBoxWidget = (ToolbarRow.ToolbarWidget.IsValid() && ToolbarRow.ToolbarWidget->GetParentWidget().IsValid()) ?
+					ToolbarRow.ToolbarWidget->GetParentWidget()->GetParentWidget() : TSharedPtr<SWidget>();
+				if (ExpanderVBoxWidget.IsValid() && ExpanderVBoxWidget->GetTypeAsString().Compare(TEXT("SVerticalBox")) == 0)
+				{
+					TSharedPtr<SVerticalBox> ExpanderVBox = StaticCastSharedPtr<SVerticalBox>(ExpanderVBoxWidget);
+					if (ExpanderVBox.IsValid() && ExpanderVBox->NumSlots() > 0)
+					{
+						const TSharedRef<SWidget>& SlotWidgetRef = ExpanderVBox->GetSlot(0).GetWidget();
+						TSharedPtr<SWidget> SlotWidgetPtr(SlotWidgetRef);
+						if (SlotWidgetPtr.IsValid() && SlotWidgetPtr->GetTypeAsString().Compare(TEXT("SBorder")) == 0)
+						{
+							TSharedPtr<SBorder> TopBorder = StaticCastSharedPtr<SBorder>(SlotWidgetPtr);
+							if (TopBorder.IsValid())
+							{
+								TopBorder->SetBorderImage(new FSlateRoundedBoxBrush(FSlateColor(ToolColor.Color), HeaderRadii));
+							}
+						}
+					}
+				}
+				break;
+			}
+		}
+	}
+}
+
+
+
 void FModelingToolsEditorModeToolkit::OnToolPaletteChanged(FName PaletteName) 
 {
 }
diff --git a/Engine/Plugins/Editor/ModelingToolsEditorMode/Source/ModelingToolsEditorMode/Private/ModelingToolsManagerActions.cpp b/Engine/Plugins/Editor/ModelingToolsEditorMode/Source/ModelingToolsEditorMode/Private/ModelingToolsManagerActions.cpp
--- a/Engine/Plugins/Editor/ModelingToolsEditorMode/Source/ModelingToolsEditorMode/Private/ModelingToolsManagerActions.cpp
+++ b/Engine/Plugins/Editor/ModelingToolsEditorMode/Source/ModelingToolsEditorMode/Private/ModelingToolsManagerActions.cpp
@@ -20,113 +20,135 @@ FModelingToolsManagerCommands::FModelingToolsManagerCommands() :
 }
 
 
-void FModelingToolsManagerCommands::RegisterCommands()
+
+TSharedPtr<FUICommandInfo> FModelingToolsManagerCommands::FindToolByName(FString Name, bool& bFound) const
 {
-	UI_COMMAND(BeginAddBoxPrimitiveTool, "Box", "Create new Box StaticMesh Assets", EUserInterfaceActionType::ToggleButton, FInputChord());
-	UI_COMMAND(BeginAddCylinderPrimitiveTool, "Cyl", "Create new Cylinder StaticMesh Assets", EUserInterfaceActionType::ToggleButton, FInputChord());
-	UI_COMMAND(BeginAddConePrimitiveTool, "Cone", "Create new Cone StaticMesh Assets", EUserInterfaceActionType::ToggleButton, FInputChord());
-	UI_COMMAND(BeginAddArrowPrimitiveTool, "Arrow", "Create new Arrow StaticMesh Assets", EUserInterfaceActionType::ToggleButton, FInputChord());
-	UI_COMMAND(BeginAddRectanglePrimitiveTool, "Rect", "Create new Rectangle StaticMesh Assets", EUserInterfaceActionType::ToggleButton, FInputChord());
-	UI_COMMAND(BeginAddDiscPrimitiveTool, "Disc", "Create new Disc StaticMesh Assets", EUserInterfaceActionType::ToggleButton, FInputChord());
-	UI_COMMAND(BeginAddTorusPrimitiveTool, "Torus", "Create new Torus StaticMesh Assets", EUserInterfaceActionType::ToggleButton, FInputChord());
-	UI_COMMAND(BeginAddSpherePrimitiveTool, "Sphere", "Create new Sphere StaticMesh Assets", EUserInterfaceActionType::ToggleButton, FInputChord());
-	UI_COMMAND(BeginAddStairsPrimitiveTool, "Stair", "Create new Stairs StaticMesh Assets", EUserInterfaceActionType::ToggleButton, FInputChord());
-
-	UI_COMMAND(BeginAddPatchTool, "Patch", "Add Patch", EUserInterfaceActionType::ToggleButton, FInputChord());
-	UI_COMMAND(BeginRevolveBoundaryTool, "BdryRev", "Revolve Mesh Boundary Loop into New Asset", EUserInterfaceActionType::ToggleButton, FInputChord());
-	UI_COMMAND(BeginDrawPolygonTool, "PolyExt", "Draw/Extrude Polygons to create new StaticMesh Assets", EUserInterfaceActionType::ToggleButton, FInputChord());
-	UI_COMMAND(BeginDrawPolyPathTool, "PathExt", "Draw/Extrude PolyPaths to create new StaticMesh Assets", EUserInterfaceActionType::ToggleButton, FInputChord());
-	UI_COMMAND(BeginDrawAndRevolveTool, "PathRev", "Draw and Revolve a PolyPath to create a new StaticMesh Assets", EUserInterfaceActionType::ToggleButton, FInputChord());
-	UI_COMMAND(BeginShapeSprayTool, "Spray", "Shape Spray", EUserInterfaceActionType::ToggleButton, FInputChord());
-
-	UI_COMMAND(BeginSculptMeshTool, "VSclpt", "Vertex Sculpting", EUserInterfaceActionType::ToggleButton, FInputChord());
-	UI_COMMAND(BeginCubeGridTool, "CubeGr", "Create block-like geometry from a repositionable grid", EUserInterfaceActionType::ToggleButton, FInputChord());
-	UI_COMMAND(BeginPolyEditTool, "PolyEd", "Edit Mesh via PolyGroups", EUserInterfaceActionType::ToggleButton, FInputChord());
-	UI_COMMAND(BeginTriEditTool, "TriEd", "Edit Mesh via Triangles", EUserInterfaceActionType::ToggleButton, FInputChord());
-	UI_COMMAND(BeginPolyDeformTool, "PolyDef", "Deform Mesh via PolyGroups", EUserInterfaceActionType::ToggleButton, FInputChord());
-	UI_COMMAND(BeginSmoothMeshTool, "Smooth", "Smooth Mesh surface", EUserInterfaceActionType::ToggleButton, FInputChord());
-	UI_COMMAND(BeginOffsetMeshTool, "Offset", "Offset Mesh surface", EUserInterfaceActionType::ToggleButton, FInputChord());
-	UI_COMMAND(BeginDisplaceMeshTool, "Displce", "Displace Mesh surface with optional subdivision", EUserInterfaceActionType::ToggleButton, FInputChord());
-	UI_COMMAND(BeginMeshSpaceDeformerTool, "Warp", "Reshape Mesh using Space Deformers", EUserInterfaceActionType::ToggleButton, FInputChord());
-	UI_COMMAND(BeginTransformMeshesTool, "XForm", "Transform selected Meshes", EUserInterfaceActionType::ToggleButton, FInputChord());
-	UI_COMMAND(BeginEditPivotTool, "Pivot", "Edit Mesh Pivots", EUserInterfaceActionType::ToggleButton, FInputChord());
-	UI_COMMAND(BeginAddPivotActorTool, "PivotAct", "Add actor to act as a pivot for child component(s)", EUserInterfaceActionType::ToggleButton, FInputChord());
-	UI_COMMAND(BeginBakeTransformTool, "BakeRS", "Bake Scale/Rotation into Mesh Asset", EUserInterfaceActionType::ToggleButton, FInputChord());
-	UI_COMMAND(BeginCombineMeshesTool, "Append", "Combine Selection into new StaticMesh Asset", EUserInterfaceActionType::ToggleButton, FInputChord());
-	UI_COMMAND(BeginDuplicateMeshesTool, "Dupe", "Duplicate Selection into new StaticMesh Asset", EUserInterfaceActionType::ToggleButton, FInputChord());
-	UI_COMMAND(BeginAlignObjectsTool, "Align", "Align Objects", EUserInterfaceActionType::ToggleButton, FInputChord());
-	UI_COMMAND(BeginTransferMeshTool, "Transfer", "Transfer Mesh", EUserInterfaceActionType::ToggleButton, FInputChord());
-	UI_COMMAND(BeginConvertMeshesTool, "Convert", "Convert Meshes", EUserInterfaceActionType::ToggleButton, FInputChord());
-	UI_COMMAND(BeginSplitMeshesTool, "Split", "Split Meshes", EUserInterfaceActionType::ToggleButton, FInputChord());
-
-	UI_COMMAND(BeginRemeshSculptMeshTool, "DSclpt", "Dynamic Mesh Sculpting", EUserInterfaceActionType::ToggleButton, FInputChord());
-	UI_COMMAND(BeginRemeshMeshTool, "Remesh", "Retriangulate Mesh", EUserInterfaceActionType::ToggleButton, FInputChord());
-	UI_COMMAND(BeginProjectToTargetTool, "Project", "Map/Remesh onto Target Mesh (second selection)", EUserInterfaceActionType::ToggleButton, FInputChord());
-	UI_COMMAND(BeginSimplifyMeshTool, "Simplfy", "Simplify Mesh", EUserInterfaceActionType::ToggleButton, FInputChord());
-	UI_COMMAND(BeginEditNormalsTool, "Nrmls", "Modify/Update Mesh Normals", EUserInterfaceActionType::ToggleButton, FInputChord());
-	UI_COMMAND(BeginEditTangentsTool, "Tngnts", "Update Mesh Tangents", EUserInterfaceActionType::ToggleButton, FInputChord());
-	UI_COMMAND(BeginRemoveOccludedTrianglesTool, "Jacket", "Remove Hidden Triangles from selected Meshes", EUserInterfaceActionType::ToggleButton, FInputChord());
-	UI_COMMAND(BeginUVProjectionTool, "Project", "Set UVs from Projection", EUserInterfaceActionType::ToggleButton, FInputChord());
-	UI_COMMAND(BeginUVLayoutTool, "Layout", "Re-Pack Existing UVs", EUserInterfaceActionType::ToggleButton, FInputChord());
-	UI_COMMAND(BeginVoxelMergeTool, "VoxMrg", "Merge Selected Meshes (Voxel-Based)", EUserInterfaceActionType::ToggleButton, FInputChord());
-	UI_COMMAND(BeginVoxelBooleanTool, "VoxBool", "Boolean Selected Meshes (Voxel-Based)", EUserInterfaceActionType::ToggleButton, FInputChord());
-	UI_COMMAND(BeginVoxelSolidifyTool, "VoxWrap", "Wrap Selected Meshes (Voxel-Based)", EUserInterfaceActionType::ToggleButton, FInputChord());
-	UI_COMMAND(BeginVoxelBlendTool, "VoxBlnd", "Blend Selected Meshes (Voxel-Based)", EUserInterfaceActionType::ToggleButton, FInputChord());
-	UI_COMMAND(BeginVoxelMorphologyTool, "VoxMrph", "Offset/Inset Selected Meshes (Voxel-Based)", EUserInterfaceActionType::ToggleButton, FInputChord());
-	UI_COMMAND(BeginSelfUnionTool, "Merge", "Self-Union Selected Meshes to resolve Self-Intersections", EUserInterfaceActionType::ToggleButton, FInputChord());
-	UI_COMMAND(BeginMeshBooleanTool, "Boolean", "Apply Boolean/CSG operation to selected Mesh pair", EUserInterfaceActionType::ToggleButton, FInputChord());
-	UI_COMMAND(BeginMeshTrimTool, "Trim", "Trim/Cut selected mesh with second mesh", EUserInterfaceActionType::ToggleButton, FInputChord());
-	UI_COMMAND(BeginCutMeshWithMeshTool, "MshCut", "Split selected mesh into parts using second mesh", EUserInterfaceActionType::ToggleButton, FInputChord());
-	UI_COMMAND(BeginBspConversionTool, "BSPConv", "Convert BSP to StaticMesh Asset", EUserInterfaceActionType::ToggleButton, FInputChord());
-	UI_COMMAND(BeginMeshToVolumeTool, "Msh2Vol", "Convert Mesh to Volume", EUserInterfaceActionType::ToggleButton, FInputChord());
-	UI_COMMAND(BeginVolumeToMeshTool, "Vol2Msh", "Convert Volume to new Mesh Asset", EUserInterfaceActionType::ToggleButton, FInputChord());
-	UI_COMMAND(BeginPlaneCutTool, "PlnCut", "Cut Selected Meshes with Plane", EUserInterfaceActionType::ToggleButton, FInputChord());
-	UI_COMMAND(BeginMirrorTool, "Mirror", "Mirror Selected Meshes", EUserInterfaceActionType::ToggleButton, FInputChord());
-	UI_COMMAND(BeginHoleFillTool, "HFill", "Fill Holes in Mesh", EUserInterfaceActionType::ToggleButton, FInputChord());
-	UI_COMMAND(BeginLatticeDeformerTool, "Lattice", "Deform Mesh with 3D Lattice/Grid", EUserInterfaceActionType::ToggleButton, FInputChord());
-	UI_COMMAND(BeginPolygonCutTool, "PolyCut", "Cut Mesh with Extruded Polygon", EUserInterfaceActionType::ToggleButton, FInputChord());
-	UI_COMMAND(BeginGlobalUVGenerateTool, "AutoUV", "Auto-Unwrap and Pack UVs for Mesh", EUserInterfaceActionType::ToggleButton, FInputChord());
-	UI_COMMAND(BeginGroupUVGenerateTool, "Unwrap", "Recalculate UV unwrapping for Mesh regions", EUserInterfaceActionType::ToggleButton, FInputChord());
-	UI_COMMAND(BeginMeshSelectionTool, "TriSel", "Select and Edit Mesh Triangles", EUserInterfaceActionType::ToggleButton, FInputChord());
-
-	UI_COMMAND(BeginPhysicsInspectorTool, "PInspct", "Inspect Physics Geometry for selected Meshes", EUserInterfaceActionType::ToggleButton, FInputChord());
-	UI_COMMAND(BeginSetCollisionGeometryTool, "Msh2Coll", "Convert Selected Meshes to Simple Collision Geometry (for last Selected)", EUserInterfaceActionType::ToggleButton, FInputChord());
-	UI_COMMAND(BeginEditCollisionGeometryTool, "EditPhys", "Edit Simple Collision Geometry for selected Mesh", EUserInterfaceActionType::ToggleButton, FInputChord());
-	UI_COMMAND(BeginExtractCollisionGeometryTool, "Coll2Msh", "Convert Simple Collision Geometry to new Mesh Asset", EUserInterfaceActionType::ToggleButton, FInputChord());
-
-	UI_COMMAND(BeginMeshInspectorTool, "Inspct", "Inspect Mesh Attributes", EUserInterfaceActionType::ToggleButton, FInputChord());
-	UI_COMMAND(BeginWeldEdgesTool, "Weld", "Weld Overlapping Mesh Edges", EUserInterfaceActionType::ToggleButton, FInputChord());
-	UI_COMMAND(BeginPolyGroupsTool, "GenGrps", "Generate new PolyGroups", EUserInterfaceActionType::ToggleButton, FInputChord());
-	UI_COMMAND(BeginEditMeshMaterialsTool, "MatEd", "Assign Materials to Selected Triangles", EUserInterfaceActionType::ToggleButton, FInputChord());
-	UI_COMMAND(BeginTransformUVIslandsTool, "XForm", "Transform UV Islands in UV Space", EUserInterfaceActionType::ToggleButton, FInputChord());
-	UI_COMMAND(BeginMeshAttributePaintTool, "MapPnt", "Paint Attribute Maps", EUserInterfaceActionType::ToggleButton, FInputChord());
-	UI_COMMAND(BeginAttributeEditorTool, "AttrEd", "Edit/Configure Mesh Attributes", EUserInterfaceActionType::ToggleButton, FInputChord());
-	UI_COMMAND(BeginBakeMeshAttributeMapsTool, "BakeTx", "Bake Image Maps for Target Mesh (optionally from second Source Mesh)", EUserInterfaceActionType::ToggleButton, FInputChord());
-	UI_COMMAND(BeginBakeMultiMeshAttributeMapsTool, "BakeAll", "Bake Image Maps for Target Mesh from multiple Source/Detail meshes", EUserInterfaceActionType::ToggleButton, FInputChord());
-	UI_COMMAND(BeginBakeMeshAttributeVertexTool, "BakeVtx", "Bake Vertex Colors for Target Mesh (optionally from a second Source Mesh)", EUserInterfaceActionType::ToggleButton, FInputChord());
-	UI_COMMAND(BeginUVSeamEditTool, "SeamEd", "Add UV Seams to Mesh", EUserInterfaceActionType::ToggleButton, FInputChord());
+	bFound = false;
+	for (const FStartToolCommand& Command : RegisteredTools)
+	{
+		if (Command.ToolUIName.Compare(Name, ESearchCase::IgnoreCase) == 0)
+		{
+			bFound = true;
+			return Command.ToolCommand;
+		}
+	}
+	return TSharedPtr<FUICommandInfo>();
+}
 
+
+void FModelingToolsManagerCommands::RegisterCommands()
+{
+	// this has to be done with a compile-time macro because UI_COMMAND expands to LOCTEXT macros
+#define REGISTER_MODELING_TOOL_COMMAND(ToolCommandInfo, ToolName, ToolTip) \
+		UI_COMMAND(ToolCommandInfo, ToolName, ToolTip, EUserInterfaceActionType::ToggleButton, FInputChord()); \
+		RegisteredTools.Add(FStartToolCommand{ ToolName, ToolCommandInfo });
+
+	REGISTER_MODELING_TOOL_COMMAND(BeginAddBoxPrimitiveTool, "Box", "Create new Box StaticMesh Assets");
+	REGISTER_MODELING_TOOL_COMMAND(BeginAddCylinderPrimitiveTool, "Cyl", "Create new Cylinder StaticMesh Assets");
+	REGISTER_MODELING_TOOL_COMMAND(BeginAddConePrimitiveTool, "Cone", "Create new Cone StaticMesh Assets");
+	REGISTER_MODELING_TOOL_COMMAND(BeginAddArrowPrimitiveTool, "Arrow", "Create new Arrow StaticMesh Assets");
+	REGISTER_MODELING_TOOL_COMMAND(BeginAddRectanglePrimitiveTool, "Rect", "Create new Rectangle StaticMesh Assets");
+	REGISTER_MODELING_TOOL_COMMAND(BeginAddDiscPrimitiveTool, "Disc", "Create new Disc StaticMesh Assets");
+	REGISTER_MODELING_TOOL_COMMAND(BeginAddTorusPrimitiveTool, "Torus", "Create new Torus StaticMesh Assets");
+	REGISTER_MODELING_TOOL_COMMAND(BeginAddSpherePrimitiveTool, "Sphere", "Create new Sphere StaticMesh Assets");
+	REGISTER_MODELING_TOOL_COMMAND(BeginAddStairsPrimitiveTool, "Stair", "Create new Stairs StaticMesh Assets");
+
+	REGISTER_MODELING_TOOL_COMMAND(BeginAddPatchTool, "Patch", "Add Patch");
+	REGISTER_MODELING_TOOL_COMMAND(BeginRevolveBoundaryTool, "BdryRev", "Revolve Mesh Boundary Loop into New Asset");
+	REGISTER_MODELING_TOOL_COMMAND(BeginDrawPolygonTool, "PolyExt", "Draw/Extrude Polygons to create new StaticMesh Assets");
+	REGISTER_MODELING_TOOL_COMMAND(BeginDrawPolyPathTool, "PathExt", "Draw/Extrude PolyPaths to create new StaticMesh Assets");
+	REGISTER_MODELING_TOOL_COMMAND(BeginDrawAndRevolveTool, "PathRev", "Draw and Revolve a PolyPath to create a new StaticMesh Assets");
+	REGISTER_MODELING_TOOL_COMMAND(BeginShapeSprayTool, "Spray", "Shape Spray");
+
+	REGISTER_MODELING_TOOL_COMMAND(BeginSculptMeshTool, "VSclpt", "Vertex Sculpting");
+	REGISTER_MODELING_TOOL_COMMAND(BeginCubeGridTool, "CubeGr", "Create block-like geometry from a repositionable grid");
+	REGISTER_MODELING_TOOL_COMMAND(BeginPolyEditTool, "PolyEd", "Edit Mesh via PolyGroups");
+	REGISTER_MODELING_TOOL_COMMAND(BeginTriEditTool, "TriEd", "Edit Mesh via Triangles");
+	REGISTER_MODELING_TOOL_COMMAND(BeginPolyDeformTool, "PolyDef", "Deform Mesh via PolyGroups");
+	REGISTER_MODELING_TOOL_COMMAND(BeginSmoothMeshTool, "Smooth", "Smooth Mesh surface");
+	REGISTER_MODELING_TOOL_COMMAND(BeginOffsetMeshTool, "Offset", "Offset Mesh surface");
+	REGISTER_MODELING_TOOL_COMMAND(BeginDisplaceMeshTool, "Displce", "Displace Mesh surface with optional subdivision");
+	REGISTER_MODELING_TOOL_COMMAND(BeginMeshSpaceDeformerTool, "Warp", "Reshape Mesh using Space Deformers");
+	REGISTER_MODELING_TOOL_COMMAND(BeginTransformMeshesTool, "XForm", "Transform selected Meshes");
+	REGISTER_MODELING_TOOL_COMMAND(BeginEditPivotTool, "Pivot", "Edit Mesh Pivots");
+	REGISTER_MODELING_TOOL_COMMAND(BeginAddPivotActorTool, "PivotAct", "Add actor to act as a pivot for child component(s)");
+	REGISTER_MODELING_TOOL_COMMAND(BeginBakeTransformTool, "BakeRS", "Bake Scale/Rotation into Mesh Asset");
+	REGISTER_MODELING_TOOL_COMMAND(BeginCombineMeshesTool, "Append", "Combine Selection into new StaticMesh Asset");
+	REGISTER_MODELING_TOOL_COMMAND(BeginDuplicateMeshesTool, "Dupe", "Duplicate Selection into new StaticMesh Asset");
+	REGISTER_MODELING_TOOL_COMMAND(BeginAlignObjectsTool, "Align", "Align Objects");
+	REGISTER_MODELING_TOOL_COMMAND(BeginTransferMeshTool, "Transfer", "Transfer Mesh");
+	REGISTER_MODELING_TOOL_COMMAND(BeginConvertMeshesTool, "Convert", "Convert Meshes");
+	REGISTER_MODELING_TOOL_COMMAND(BeginSplitMeshesTool, "Split", "Split Meshes");
+
+	REGISTER_MODELING_TOOL_COMMAND(BeginRemeshSculptMeshTool, "DSclpt", "Dynamic Mesh Sculpting");
+	REGISTER_MODELING_TOOL_COMMAND(BeginRemeshMeshTool, "Remesh", "Retriangulate Mesh");
+	REGISTER_MODELING_TOOL_COMMAND(BeginProjectToTargetTool, "Project", "Map/Remesh onto Target Mesh (second selection)");
+	REGISTER_MODELING_TOOL_COMMAND(BeginSimplifyMeshTool, "Simplfy", "Simplify Mesh");
+	REGISTER_MODELING_TOOL_COMMAND(BeginEditNormalsTool, "Nrmls", "Modify/Update Mesh Normals");
+	REGISTER_MODELING_TOOL_COMMAND(BeginEditTangentsTool, "Tngnts", "Update Mesh Tangents");
+	REGISTER_MODELING_TOOL_COMMAND(BeginRemoveOccludedTrianglesTool, "Jacket", "Remove Hidden Triangles from selected Meshes");
+	REGISTER_MODELING_TOOL_COMMAND(BeginUVProjectionTool, "Project", "Set UVs from Projection");
+	REGISTER_MODELING_TOOL_COMMAND(BeginUVLayoutTool, "Layout", "Re-Pack Existing UVs");
+	REGISTER_MODELING_TOOL_COMMAND(BeginVoxelMergeTool, "VoxMrg", "Merge Selected Meshes (Voxel-Based)");
+	REGISTER_MODELING_TOOL_COMMAND(BeginVoxelBooleanTool, "VoxBool", "Boolean Selected Meshes (Voxel-Based)");
+	REGISTER_MODELING_TOOL_COMMAND(BeginVoxelSolidifyTool, "VoxWrap", "Wrap Selected Meshes (Voxel-Based)");
+	REGISTER_MODELING_TOOL_COMMAND(BeginVoxelBlendTool, "VoxBlnd", "Blend Selected Meshes (Voxel-Based)");
+	REGISTER_MODELING_TOOL_COMMAND(BeginVoxelMorphologyTool, "VoxMrph", "Offset/Inset Selected Meshes (Voxel-Based)");
+	REGISTER_MODELING_TOOL_COMMAND(BeginSelfUnionTool, "Merge", "Self-Union Selected Meshes to resolve Self-Intersections");
+	REGISTER_MODELING_TOOL_COMMAND(BeginMeshBooleanTool, "Boolean", "Apply Boolean/CSG operation to selected Mesh pair");
+	REGISTER_MODELING_TOOL_COMMAND(BeginMeshTrimTool, "Trim", "Trim/Cut selected mesh with second mesh");
+	REGISTER_MODELING_TOOL_COMMAND(BeginCutMeshWithMeshTool, "MshCut", "Split selected mesh into parts using second mesh");
+	REGISTER_MODELING_TOOL_COMMAND(BeginBspConversionTool, "BSPConv", "Convert BSP to StaticMesh Asset");
+	REGISTER_MODELING_TOOL_COMMAND(BeginMeshToVolumeTool, "Msh2Vol", "Convert Mesh to Volume");
+	REGISTER_MODELING_TOOL_COMMAND(BeginVolumeToMeshTool, "Vol2Msh", "Convert Volume to new Mesh Asset");
+	REGISTER_MODELING_TOOL_COMMAND(BeginPlaneCutTool, "PlnCut", "Cut Selected Meshes with Plane");
+	REGISTER_MODELING_TOOL_COMMAND(BeginMirrorTool, "Mirror", "Mirror Selected Meshes");
+	REGISTER_MODELING_TOOL_COMMAND(BeginHoleFillTool, "HFill", "Fill Holes in Mesh");
+	REGISTER_MODELING_TOOL_COMMAND(BeginLatticeDeformerTool, "Lattice", "Deform Mesh with 3D Lattice/Grid");
+	REGISTER_MODELING_TOOL_COMMAND(BeginPolygonCutTool, "PolyCut", "Cut Mesh with Extruded Polygon");
+	REGISTER_MODELING_TOOL_COMMAND(BeginGlobalUVGenerateTool, "AutoUV", "Auto-Unwrap and Pack UVs for Mesh");
+	REGISTER_MODELING_TOOL_COMMAND(BeginGroupUVGenerateTool, "Unwrap", "Recalculate UV unwrapping for Mesh regions");
+	REGISTER_MODELING_TOOL_COMMAND(BeginMeshSelectionTool, "TriSel", "Select and Edit Mesh Triangles");
+
+	REGISTER_MODELING_TOOL_COMMAND(BeginPhysicsInspectorTool, "PInspct", "Inspect Physics Geometry for selected Meshes");
+	REGISTER_MODELING_TOOL_COMMAND(BeginSetCollisionGeometryTool, "Msh2Coll", "Convert Selected Meshes to Simple Collision Geometry (for last Selected)");
+	REGISTER_MODELING_TOOL_COMMAND(BeginEditCollisionGeometryTool, "EditPhys", "Edit Simple Collision Geometry for selected Mesh");
+	REGISTER_MODELING_TOOL_COMMAND(BeginExtractCollisionGeometryTool, "Coll2Msh", "Convert Simple Collision Geometry to new Mesh Asset");
+
+	REGISTER_MODELING_TOOL_COMMAND(BeginMeshInspectorTool, "Inspct", "Inspect Mesh Attributes");
+	REGISTER_MODELING_TOOL_COMMAND(BeginWeldEdgesTool, "Weld", "Weld Overlapping Mesh Edges");
+	REGISTER_MODELING_TOOL_COMMAND(BeginPolyGroupsTool, "GenGrps", "Generate new PolyGroups");
+	REGISTER_MODELING_TOOL_COMMAND(BeginEditMeshMaterialsTool, "MatEd", "Assign Materials to Selected Triangles");
+	REGISTER_MODELING_TOOL_COMMAND(BeginTransformUVIslandsTool, "XForm", "Transform UV Islands in UV Space");
+	REGISTER_MODELING_TOOL_COMMAND(BeginMeshAttributePaintTool, "MapPnt", "Paint Attribute Maps");
+	REGISTER_MODELING_TOOL_COMMAND(BeginAttributeEditorTool, "AttrEd", "Edit/Configure Mesh Attributes");
+	REGISTER_MODELING_TOOL_COMMAND(BeginBakeMeshAttributeMapsTool, "BakeTx", "Bake Image Maps for Target Mesh (optionally from second Source Mesh)");
+	REGISTER_MODELING_TOOL_COMMAND(BeginBakeMultiMeshAttributeMapsTool, "BakeAll", "Bake Image Maps for Target Mesh from multiple Source/Detail meshes");
+	REGISTER_MODELING_TOOL_COMMAND(BeginBakeMeshAttributeVertexTool, "BakeVtx", "Bake Vertex Colors for Target Mesh (optionally from a second Source Mesh)");
+	REGISTER_MODELING_TOOL_COMMAND(BeginUVSeamEditTool, "SeamEd", "Add UV Seams to Mesh");
+
+	// why are these ::Button ?
 	UI_COMMAND(BeginSkinWeightsPaintTool, "SkinWts", "Start the Paint Skin Weights Tool", EUserInterfaceActionType::Button, FInputChord());
 	UI_COMMAND(BeginSkinWeightsBindingTool, "SkinBind", "Start the Skin Weights Binding Tool", EUserInterfaceActionType::Button, FInputChord());
 
-	UI_COMMAND(BeginSubdividePolyTool, "SubD", "Subdivide Mesh Polygroups or Triangles", EUserInterfaceActionType::ToggleButton, FInputChord());
+	REGISTER_MODELING_TOOL_COMMAND(BeginSubdividePolyTool, "SubD", "Subdivide Mesh Polygroups or Triangles");
 	
-	UI_COMMAND(BeginLODManagerTool, "LODMgr", "Static Mesh Asset LOD Manager", EUserInterfaceActionType::ToggleButton, FInputChord());
-	UI_COMMAND(BeginGenerateStaticMeshLODAssetTool, "AutoLOD", "Generate Static Mesh LOD Asset", EUserInterfaceActionType::ToggleButton, FInputChord());
-	UI_COMMAND(BeginMeshGroupPaintTool, "GrpPnt", "Paint New Mesh Polygroups", EUserInterfaceActionType::ToggleButton, FInputChord());
-
-
-	UI_COMMAND(BeginPolyModelTool_FaceSelect, "Faces", "PolyGroup Face Selection Tool", EUserInterfaceActionType::ToggleButton, FInputChord());
-	UI_COMMAND(BeginPolyModelTool_EdgeSelect, "Edges", "PolyGroup Edge Selection Tool", EUserInterfaceActionType::ToggleButton, FInputChord());
-	UI_COMMAND(BeginPolyModelTool_VertexSelect, "Verts", "PolyGroup Vertex Selection Tool", EUserInterfaceActionType::ToggleButton, FInputChord());
-	UI_COMMAND(BeginPolyModelTool_AllSelect, "Any", "PolyGroup Face/Edge/Vertex Selection Tool", EUserInterfaceActionType::ToggleButton, FInputChord());
-	UI_COMMAND(BeginPolyModelTool_LoopSelect, "Loops", "PolyGroup Loop Selection Tool", EUserInterfaceActionType::ToggleButton, FInputChord());
-	UI_COMMAND(BeginPolyModelTool_RingSelect, "Rings", "PolyGroup Ring Selection Tool", EUserInterfaceActionType::ToggleButton, FInputChord());
-
-	UI_COMMAND(BeginPolyModelTool_Extrude, "Extrude", "PolyGroup Extrude Tool", EUserInterfaceActionType::ToggleButton, FInputChord());
-	UI_COMMAND(BeginPolyModelTool_Inset, "Inset", "PolyGroup Inset Tool", EUserInterfaceActionType::ToggleButton, FInputChord());
-	UI_COMMAND(BeginPolyModelTool_Outset, "Outset", "PolyGroup Outset Tool", EUserInterfaceActionType::ToggleButton, FInputChord());
-	UI_COMMAND(BeginPolyModelTool_CutFaces, "Cut", "PolyGroup Cut Faces Tool", EUserInterfaceActionType::ToggleButton, FInputChord());
+	REGISTER_MODELING_TOOL_COMMAND(BeginLODManagerTool, "LODMgr", "Static Mesh Asset LOD Manager");
+	REGISTER_MODELING_TOOL_COMMAND(BeginGenerateStaticMeshLODAssetTool, "AutoLOD", "Generate Static Mesh LOD Asset");
+	REGISTER_MODELING_TOOL_COMMAND(BeginMeshGroupPaintTool, "GrpPnt", "Paint New Mesh Polygroups");
+
+
+	REGISTER_MODELING_TOOL_COMMAND(BeginPolyModelTool_FaceSelect, "Faces", "PolyGroup Face Selection Tool");
+	REGISTER_MODELING_TOOL_COMMAND(BeginPolyModelTool_EdgeSelect, "Edges", "PolyGroup Edge Selection Tool");
+	REGISTER_MODELING_TOOL_COMMAND(BeginPolyModelTool_VertexSelect, "Verts", "PolyGroup Vertex Selection Tool");
+	REGISTER_MODELING_TOOL_COMMAND(BeginPolyModelTool_AllSelect, "Any", "PolyGroup Face/Edge/Vertex Selection Tool");
+	REGISTER_MODELING_TOOL_COMMAND(BeginPolyModelTool_LoopSelect, "Loops", "PolyGroup Loop Selection Tool");
+	REGISTER_MODELING_TOOL_COMMAND(BeginPolyModelTool_RingSelect, "Rings", "PolyGroup Ring Selection Tool");
+
+	REGISTER_MODELING_TOOL_COMMAND(BeginPolyModelTool_Extrude, "Extrude", "PolyGroup Extrude Tool");
+	REGISTER_MODELING_TOOL_COMMAND(BeginPolyModelTool_Inset, "Inset", "PolyGroup Inset Tool");
+	REGISTER_MODELING_TOOL_COMMAND(BeginPolyModelTool_Outset, "Outset", "PolyGroup Outset Tool");
+	REGISTER_MODELING_TOOL_COMMAND(BeginPolyModelTool_CutFaces, "Cut", "PolyGroup Cut Faces Tool");
 
 	UI_COMMAND(AcceptActiveTool, "Accept", "Accept the active tool", EUserInterfaceActionType::Button, FInputChord());
 	UI_COMMAND(CancelActiveTool, "Cancel", "Cancel the active tool", EUserInterfaceActionType::Button, FInputChord());
diff --git a/Engine/Plugins/Editor/ModelingToolsEditorMode/Source/ModelingToolsEditorMode/Public/ModelingToolsEditorModeSettings.h b/Engine/Plugins/Editor/ModelingToolsEditorMode/Source/ModelingToolsEditorMode/Public/ModelingToolsEditorModeSettings.h
--- a/Engine/Plugins/Editor/ModelingToolsEditorMode/Source/ModelingToolsEditorMode/Public/ModelingToolsEditorModeSettings.h
+++ b/Engine/Plugins/Editor/ModelingToolsEditorMode/Source/ModelingToolsEditorMode/Public/ModelingToolsEditorModeSettings.h
@@ -129,3 +129,56 @@ public:
 	}
 
 };
+
+
+
+
+/**
+ * Defines a color to be used for a particular Tool Palette Section
+ */
+USTRUCT()
+struct FModelingModeCustomSectionColor
+{
+	GENERATED_BODY()
+
+	/** Name of Section in Modeling Mode Tool Palette */
+	UPROPERTY(EditAnywhere, Category = "SectionColor")
+	FString SectionName;
+
+	/** Custom Header Color */
+	UPROPERTY(EditAnywhere, Category = "SectionColor")
+	FLinearColor Color;
+};
+
+
+
+UCLASS(config=Editor)
+class MODELINGTOOLSEDITORMODE_API UModelingToolsModeCustomizationSettings : public UDeveloperSettings
+{
+	GENERATED_BODY()
+
+public:
+	// UDeveloperSettings overrides
+
+	virtual FName GetContainerName() const { return FName("Editor"); }
+	virtual FName GetCategoryName() const { return FName("Plugins"); }
+	virtual FName GetSectionName() const { return FName("ModelingMode"); }
+
+	virtual FText GetSectionText() const override;
+	virtual FText GetSectionDescription() const override;
+
+public:
+
+
+	/** Add the names of Modeling Mode Tool Palette Sections to have them appear at the top of the Tool Palette, in the order listed below. */
+	UPROPERTY(config, EditAnywhere, Category = "Modeling Mode|UI Customization")
+	TArray<FString> ToolSectionOrder;
+
+	/** Tool Names listed in the array below will appear in a Favorites section at the top of the Modeling Mode Tool Palette */
+	UPROPERTY(config, EditAnywhere, Category = "Modeling Mode|UI Customization")
+	TArray<FString> ToolFavorites;
+
+	/** Custom Section Header Colors for listed Sections in the Modeling Mode Tool Palette */
+	UPROPERTY(config, EditAnywhere, Category = "Modeling Mode|UI Customization")
+	TArray<FModelingModeCustomSectionColor> SectionColors;
+};
diff --git a/Engine/Plugins/Editor/ModelingToolsEditorMode/Source/ModelingToolsEditorMode/Public/ModelingToolsEditorModeToolkit.h b/Engine/Plugins/Editor/ModelingToolsEditorMode/Source/ModelingToolsEditorMode/Public/ModelingToolsEditorModeToolkit.h
--- a/Engine/Plugins/Editor/ModelingToolsEditorMode/Source/ModelingToolsEditorMode/Public/ModelingToolsEditorModeToolkit.h
+++ b/Engine/Plugins/Editor/ModelingToolsEditorMode/Source/ModelingToolsEditorMode/Public/ModelingToolsEditorModeToolkit.h
@@ -53,9 +53,6 @@ public:
 	virtual FText GetActiveToolDisplayName() const override { return ActiveToolName; }
 	virtual FText GetActiveToolMessage() const override { return ActiveToolMessage; }
 
-	virtual void BuildToolPalette_Standard(FName PaletteName, class FToolBarBuilder& ToolbarBuilder);
-	virtual void BuildToolPalette_Experimental(FName PaletteName, class FToolBarBuilder& ToolbarBuilder);
-
 	virtual void EnableShowRealtimeWarning(bool bEnable);
 
 	virtual void OnToolStarted(UInteractiveToolManager* Manager, UInteractiveTool* Tool) override;
@@ -66,6 +63,9 @@ public:
 	void OnActiveViewportChanged(TSharedPtr<IAssetViewport>, TSharedPtr<IAssetViewport> );
 
 
+
+	virtual void InvokeUI() override;
+
 private:
 	const static TArray<FName> PaletteNames_Standard;
 	const static TArray<FName> PaletteNames_Experimental;
diff --git a/Engine/Plugins/Editor/ModelingToolsEditorMode/Source/ModelingToolsEditorMode/Public/ModelingToolsManagerActions.h b/Engine/Plugins/Editor/ModelingToolsEditorMode/Source/ModelingToolsEditorMode/Public/ModelingToolsManagerActions.h
--- a/Engine/Plugins/Editor/ModelingToolsEditorMode/Source/ModelingToolsEditorMode/Public/ModelingToolsManagerActions.h
+++ b/Engine/Plugins/Editor/ModelingToolsEditorMode/Source/ModelingToolsEditorMode/Public/ModelingToolsManagerActions.h
@@ -7,13 +7,30 @@
 
 
 /**
- * TInteractiveToolCommands implementation for this module that provides standard Editor hotkey support
+ * TInteractiveToolCommands implementation for Modeling Mode Tools
  */
 class MODELINGTOOLSEDITORMODE_API FModelingToolsManagerCommands : public TCommands<FModelingToolsManagerCommands>
 {
 public:
 	FModelingToolsManagerCommands();
 
+protected:
+	struct FStartToolCommand
+	{
+		FString ToolUIName;
+		TSharedPtr<FUICommandInfo> ToolCommand;
+	};
+	TArray<FStartToolCommand> RegisteredTools;		// Tool start-commands listed below are stored in this list
+
+public:
+	/**
+	 * Find Tool start-command below by registered name (tool icon name in Mode palette)
+	 */
+	TSharedPtr<FUICommandInfo> FindToolByName(FString Name, bool& bFound) const;
+	
+	//
+	// These commands are set up to launch registered Tools via the ToolManager in ModelingToolsEditorMode.cpp
+	//
 
 	TSharedPtr<FUICommandInfo> BeginAddBoxPrimitiveTool;
 	TSharedPtr<FUICommandInfo> BeginAddCylinderPrimitiveTool;
@@ -119,6 +136,9 @@ public:
 	TSharedPtr<FUICommandInfo> BeginPolyModelTool_Outset;
 	TSharedPtr<FUICommandInfo> BeginPolyModelTool_CutFaces;
 
+	//
+	// Accept/Cancel/Complete commands are used to end the active Tool via ToolManager
+	//
 
 	TSharedPtr<FUICommandInfo> AcceptActiveTool;
 	TSharedPtr<FUICommandInfo> CancelActiveTool;
@@ -128,7 +148,7 @@ public:
 	TSharedPtr<FUICommandInfo> CancelOrCompleteActiveTool;
 
 	/**
-	 * Initialize commands
+	 * Initialize above commands
 	 */
 	virtual void RegisterCommands() override;
 };
