ModelingMode: initialize settings ustruct
#rnx

diff --git a/Engine/Plugins/Editor/ModelingToolsEditorMode/Source/ModelingToolsEditorMode/Public/ModelingToolsEditorModeSettings.h b/Engine/Plugins/Editor/ModelingToolsEditorMode/Source/ModelingToolsEditorMode/Public/ModelingToolsEditorModeSettings.h
--- a/Engine/Plugins/Editor/ModelingToolsEditorMode/Source/ModelingToolsEditorMode/Public/ModelingToolsEditorModeSettings.h
+++ b/Engine/Plugins/Editor/ModelingToolsEditorMode/Source/ModelingToolsEditorMode/Public/ModelingToolsEditorModeSettings.h
@@ -143,11 +143,11 @@ struct FModelingModeCustomSectionColor
 
 	/** Name of Section in Modeling Mode Tool Palette */
 	UPROPERTY(EditAnywhere, Category = "SectionColor")
-	FString SectionName;
+	FString SectionName = TEXT("");
 
 	/** Custom Header Color */
 	UPROPERTY(EditAnywhere, Category = "SectionColor")
-	FLinearColor Color;
+	FLinearColor Color = FLinearColor::Gray;
 };
 
 
