ModelingTools: Add icon tint customization support to Modeling & Fracture tools palette.

#rnx

diff --git a/Engine/Plugins/Editor/ModelingToolsEditorMode/Source/ModelingToolsEditorMode/Private/ModelingToolsEditorModeToolkit.cpp b/Engine/Plugins/Editor/ModelingToolsEditorMode/Source/ModelingToolsEditorMode/Private/ModelingToolsEditorModeToolkit.cpp
--- a/Engine/Plugins/Editor/ModelingToolsEditorMode/Source/ModelingToolsEditorMode/Private/ModelingToolsEditorModeToolkit.cpp
+++ b/Engine/Plugins/Editor/ModelingToolsEditorMode/Source/ModelingToolsEditorMode/Private/ModelingToolsEditorModeToolkit.cpp
@@ -22,6 +22,8 @@
 #include "Widgets/Layout/SSeparator.h"
 #include "Widgets/Images/SImage.h"
 #include "Widgets/Layout/SExpandableArea.h"
+#include "Framework/MultiBox/SToolBarButtonBlock.h"
+#include "Internationalization/Text.h"
 
 // for Tool Extensions
 #include "Features/IModularFeatures.h"
@@ -752,8 +754,31 @@ void FModelingToolsEditorModeToolkit::InvokeUI()
 		HeaderRadii = BaseBrush->OutlineSettings.CornerRadii;
 	}
 
+	// Generate a map for tool specific colors
+	TMap<FString, FLinearColor> SectionIconColorMap;
+	TMap<FString, TMap<FString, FLinearColor>> SectionToolIconColorMap;
+	for (const FModelingModeCustomToolColor& ToolColor : UISettings->ToolColors)
+	{
+		FString SectionName, ToolName;
+		ToolColor.ToolName.Split(".", &SectionName, &ToolName);
+		SectionName.ToLowerInline();
+		if (ToolName.Len() > 0)
+		{
+			if (!SectionToolIconColorMap.Contains(SectionName))
+			{
+				SectionToolIconColorMap.Emplace(SectionName, TMap<FString, FLinearColor>());
+			}
+			SectionToolIconColorMap[SectionName].Add(ToolName, ToolColor.Color);
+		}
+		else
+		{
+			SectionIconColorMap.Emplace(ToolColor.ToolName.ToLower(), ToolColor.Color);
+		}
+	}
+
 	for (FEdModeToolbarRow& ToolbarRow : ActiveToolBarRows)
 	{
+		// Update section header colors
 		for (FModelingModeCustomSectionColor ToolColor : UISettings->SectionColors)
 		{
 			if (ToolColor.SectionName.Equals(ToolbarRow.PaletteName.ToString(), ESearchCase::IgnoreCase)
@@ -785,6 +810,103 @@ void FModelingToolsEditorModeToolkit::InvokeUI()
 				break;
 			}
 		}
+
+		// Update tool colors
+		FLinearColor* SectionIconColor = SectionIconColorMap.Find(ToolbarRow.PaletteName.ToString().ToLower());
+		if (!SectionIconColor)
+		{
+			SectionIconColor = SectionIconColorMap.Find(ToolbarRow.DisplayName.ToString().ToLower());
+		}
+		TMap<FString, FLinearColor>* SectionToolIconColors = SectionToolIconColorMap.Find(ToolbarRow.PaletteName.ToString().ToLower());
+		if (!SectionToolIconColors)
+		{
+			SectionToolIconColors = SectionToolIconColorMap.Find(ToolbarRow.DisplayName.ToString().ToLower());
+		}
+		if (SectionIconColor || SectionToolIconColors)
+		{
+			// code below is highly dependent on the structure of the ToolbarRow.ToolbarWidget. Currently this is 
+			// a SMultiBoxWidget. The code will fail gracefully if this structure changes.
+			
+			if (ToolbarRow.ToolbarWidget.IsValid() && ToolbarRow.ToolbarWidget->GetTypeAsString().Compare(TEXT("SMultiBoxWidget")) == 0)
+			{
+				auto FindFirstChildWidget = [](const TSharedPtr<SWidget>& Widget, const FString& WidgetType)
+				{
+					auto FindFirstChildWidgetImpl = [](const TSharedPtr<SWidget>& Widget, const FString& WidgetType, auto& FindRef) -> TSharedPtr<SWidget>
+					{
+						TSharedPtr<SWidget> Result;
+						if (Widget.IsValid())
+						{
+							FChildren* Children = Widget->GetChildren();
+							const int32 NumChild = Children ? Children->NumSlot() : 0;
+							for (int32 ChildIdx = 0; ChildIdx < NumChild; ++ChildIdx)
+							{
+								const TSharedRef<SWidget> ChildWidgetRef = Children->GetChildAt(ChildIdx);
+								TSharedPtr<SWidget> ChildWidgetPtr(ChildWidgetRef);
+								if (ChildWidgetPtr.IsValid())
+								{
+									if (ChildWidgetPtr->GetTypeAsString().Compare(WidgetType) == 0)
+									{
+										Result = ChildWidgetPtr;
+										break;
+									}
+
+									Result = FindRef(ChildWidgetPtr, WidgetType, FindRef);
+									if (Result.IsValid())
+									{
+										break;
+									}
+								}
+							}
+						}
+						return Result;
+					};
+					return FindFirstChildWidgetImpl(Widget, WidgetType, FindFirstChildWidgetImpl);
+				};
+
+				TSharedPtr<SWidget> PanelWidget = FindFirstChildWidget(ToolbarRow.ToolbarWidget, TEXT("SUniformWrapPanel"));
+				if (PanelWidget.IsValid())
+				{
+					// This contains each of the FToolBarButtonBlock items for this row.
+					FChildren* PanelChildren = PanelWidget->GetChildren();
+					const int32 NumChild = PanelChildren ? PanelChildren->NumSlot() : 0;
+					for (int32 ChildIdx = 0; ChildIdx < NumChild; ++ChildIdx)
+					{
+						const TSharedRef<SWidget> ChildWidgetRef = PanelChildren->GetChildAt(ChildIdx);
+						TSharedPtr<SWidget> ChildWidgetPtr(ChildWidgetRef);
+						if (ChildWidgetPtr.IsValid() && ChildWidgetPtr->GetTypeAsString().Compare(TEXT("SToolBarButtonBlock")) == 0)
+						{
+							TSharedPtr<SToolBarButtonBlock> ToolBarButton = StaticCastSharedPtr<SToolBarButtonBlock>(ChildWidgetPtr);
+							if (ToolBarButton.IsValid())
+							{
+								TSharedPtr<SWidget> LayeredImageWidget = FindFirstChildWidget(ToolBarButton, TEXT("SLayeredImage"));
+								TSharedPtr<SWidget> TextBlockWidget = FindFirstChildWidget(ToolBarButton, TEXT("STextBlock"));
+								if (LayeredImageWidget.IsValid() && TextBlockWidget.IsValid())
+								{
+									TSharedPtr<SImage> ImageWidget = StaticCastSharedPtr<SImage>(LayeredImageWidget);
+									TSharedPtr<STextBlock> TextWidget = StaticCastSharedPtr<STextBlock>(TextBlockWidget);
+									// Check if this Section.Tool has an explicit color entry. If not, fallback
+									// to any Section-wide color entry, otherwise leave the tint alone.
+									FLinearColor* TintColor = SectionToolIconColors ? SectionToolIconColors->Find(TextWidget->GetText().ToString()) : nullptr;
+									if (!TintColor)
+									{
+										const FString* SourceText = FTextInspector::GetSourceString(TextWidget->GetText());
+										TintColor = SectionToolIconColors && SourceText ? SectionToolIconColors->Find(*SourceText) : nullptr;
+										if (!TintColor)
+										{
+											TintColor = SectionIconColor;
+										}
+									}
+									if (TintColor)
+									{
+										ImageWidget->SetColorAndOpacity(FSlateColor(*TintColor));
+									}
+								}
+							}
+						}
+					}
+				}
+			}
+		}
 	}
 }
 
diff --git a/Engine/Plugins/Editor/ModelingToolsEditorMode/Source/ModelingToolsEditorMode/Public/ModelingToolsEditorModeSettings.h b/Engine/Plugins/Editor/ModelingToolsEditorMode/Source/ModelingToolsEditorMode/Public/ModelingToolsEditorModeSettings.h
--- a/Engine/Plugins/Editor/ModelingToolsEditorMode/Source/ModelingToolsEditorMode/Public/ModelingToolsEditorModeSettings.h
+++ b/Engine/Plugins/Editor/ModelingToolsEditorMode/Source/ModelingToolsEditorMode/Public/ModelingToolsEditorModeSettings.h
@@ -151,6 +151,29 @@ struct FModelingModeCustomSectionColor
 };
 
 
+/**
+ * Defines a color to be used for a particular Tool Palette Tool
+ */
+USTRUCT()
+struct FModelingModeCustomToolColor
+{
+	GENERATED_BODY()
+
+	/**
+	 * Name of Section or Tool in Modeling Mode Tool Palette
+	 *
+	 * Format:
+	 * SectionName        (Specifies a default color for all tools in the section.)
+	 * SectionName.ToolName        (Specifies an override color for a specific tool in the given section.)
+	 */
+	UPROPERTY(EditAnywhere, Category = "ToolColor")
+	FString ToolName = TEXT("");
+
+	/** Custom Tool Color */
+	UPROPERTY(EditAnywhere, Category = "ToolColor")
+	FLinearColor Color = FLinearColor::Gray;
+};
+
 
 UCLASS(config=Editor)
 class MODELINGTOOLSEDITORMODE_API UModelingToolsModeCustomizationSettings : public UDeveloperSettings
@@ -181,4 +204,14 @@ public:
 	/** Custom Section Header Colors for listed Sections in the Modeling Mode Tool Palette */
 	UPROPERTY(config, EditAnywhere, Category = "Modeling Mode|UI Customization")
 	TArray<FModelingModeCustomSectionColor> SectionColors;
+
+	/**
+	 * Custom Tool Colors for listed Tools in the Modeling Mode Tool Palette.
+	 * 
+	 * Format:
+	 * SectionName        (Specifies a default color for all tools in the section.)
+	 * SectionName.ToolName        (Specifies an override color for a specific tool in the given section.)
+	 */
+	UPROPERTY(config, EditAnywhere, Category = "Modeling Mode|UI Customization")
+	TArray<FModelingModeCustomToolColor> ToolColors;
 };
diff --git a/Engine/Plugins/Experimental/ChaosEditor/Source/FractureEditor/Private/FractureEditorModeToolkit.cpp b/Engine/Plugins/Experimental/ChaosEditor/Source/FractureEditor/Private/FractureEditorModeToolkit.cpp
--- a/Engine/Plugins/Experimental/ChaosEditor/Source/FractureEditor/Private/FractureEditorModeToolkit.cpp
+++ b/Engine/Plugins/Experimental/ChaosEditor/Source/FractureEditor/Private/FractureEditorModeToolkit.cpp
@@ -29,6 +29,8 @@
 #include "Widgets/Text/STextBlock.h"
 #include "Widgets/Layout/SGridPanel.h"
 #include "Framework/MultiBox/MultiBoxBuilder.h"
+#include "Framework/MultiBox/SToolBarButtonBlock.h"
+#include "Internationalization/Text.h"
 
 #include "FractureTool.h"
 
@@ -355,8 +357,31 @@ void FFractureEditorModeToolkit::InvokeUI()
 		HeaderRadii = BaseBrush->OutlineSettings.CornerRadii;
 	}
 
+	// Generate a map for tool specific colors
+	TMap<FString, FLinearColor> SectionIconColorMap;
+	TMap<FString, TMap<FString, FLinearColor>> SectionToolIconColorMap;
+	for (const FFractureModeCustomToolColor& ToolColor : UISettings->ToolColors)
+	{
+		FString SectionName, ToolName;
+		ToolColor.ToolName.Split(".", &SectionName, &ToolName);
+		SectionName.ToLowerInline();
+		if (ToolName.Len() > 0)
+		{
+			if (!SectionToolIconColorMap.Contains(SectionName))
+			{
+				SectionToolIconColorMap.Emplace(SectionName, TMap<FString, FLinearColor>());
+			}
+			SectionToolIconColorMap[SectionName].Add(ToolName, ToolColor.Color);
+		}
+		else
+		{
+			SectionIconColorMap.Emplace(ToolColor.ToolName.ToLower(), ToolColor.Color);
+		}
+	}
+
 	for (FEdModeToolbarRow& ToolbarRow : ActiveToolBarRows)
 	{
+		// Update section header colors
 		for (FFractureModeCustomSectionColor ToolColor : UISettings->SectionColors)
 		{
 			if (ToolColor.SectionName.Equals(ToolbarRow.DisplayName.ToString(), ESearchCase::IgnoreCase)
@@ -388,6 +413,103 @@ void FFractureEditorModeToolkit::InvokeUI()
 				break;
 			}
 		}
+
+		// Update tool colors
+		FLinearColor* SectionIconColor = SectionIconColorMap.Find(ToolbarRow.PaletteName.ToString().ToLower());
+		if (!SectionIconColor)
+		{
+			SectionIconColor = SectionIconColorMap.Find(ToolbarRow.DisplayName.ToString().ToLower());
+		}
+		TMap<FString, FLinearColor>* SectionToolIconColors = SectionToolIconColorMap.Find(ToolbarRow.PaletteName.ToString().ToLower());
+		if (!SectionToolIconColors)
+		{
+			SectionToolIconColors = SectionToolIconColorMap.Find(ToolbarRow.DisplayName.ToString().ToLower());
+		}
+		if (SectionIconColor || SectionToolIconColors)
+		{
+			// code below is highly dependent on the structure of the ToolbarRow.ToolbarWidget. Currently this is 
+			// a SMultiBoxWidget. The code will fail gracefully if this structure changes.
+			
+			if (ToolbarRow.ToolbarWidget.IsValid() && ToolbarRow.ToolbarWidget->GetTypeAsString().Compare(TEXT("SMultiBoxWidget")) == 0)
+			{
+				auto FindFirstChildWidget = [](const TSharedPtr<SWidget>& Widget, const FString& WidgetType)
+				{
+					auto FindFirstChildWidgetImpl = [](const TSharedPtr<SWidget>& Widget, const FString& WidgetType, auto& FindRef) -> TSharedPtr<SWidget>
+					{
+						TSharedPtr<SWidget> Result;
+						if (Widget.IsValid())
+						{
+							FChildren* Children = Widget->GetChildren();
+							const int32 NumChild = Children ? Children->NumSlot() : 0;
+							for (int32 ChildIdx = 0; ChildIdx < NumChild; ++ChildIdx)
+							{
+								const TSharedRef<SWidget> ChildWidgetRef = Children->GetChildAt(ChildIdx);
+								TSharedPtr<SWidget> ChildWidgetPtr(ChildWidgetRef);
+								if (ChildWidgetPtr.IsValid())
+								{
+									if (ChildWidgetPtr->GetTypeAsString().Compare(WidgetType) == 0)
+									{
+										Result = ChildWidgetPtr;
+										break;
+									}
+
+									Result = FindRef(ChildWidgetPtr, WidgetType, FindRef);
+									if (Result.IsValid())
+									{
+										break;
+									}
+								}
+							}
+						}
+						return Result;
+					};
+					return FindFirstChildWidgetImpl(Widget, WidgetType, FindFirstChildWidgetImpl);
+				};
+
+				TSharedPtr<SWidget> PanelWidget = FindFirstChildWidget(ToolbarRow.ToolbarWidget, TEXT("SUniformWrapPanel"));
+				if (PanelWidget.IsValid())
+				{
+					// This contains each of the FToolBarButtonBlock items for this row.
+					FChildren* PanelChildren = PanelWidget->GetChildren();
+					const int32 NumChild = PanelChildren ? PanelChildren->NumSlot() : 0;
+					for (int32 ChildIdx = 0; ChildIdx < NumChild; ++ChildIdx)
+					{
+						const TSharedRef<SWidget> ChildWidgetRef = PanelChildren->GetChildAt(ChildIdx);
+						TSharedPtr<SWidget> ChildWidgetPtr(ChildWidgetRef);
+						if (ChildWidgetPtr.IsValid() && ChildWidgetPtr->GetTypeAsString().Compare(TEXT("SToolBarButtonBlock")) == 0)
+						{
+							TSharedPtr<SToolBarButtonBlock> ToolBarButton = StaticCastSharedPtr<SToolBarButtonBlock>(ChildWidgetPtr);
+							if (ToolBarButton.IsValid())
+							{
+								TSharedPtr<SWidget> LayeredImageWidget = FindFirstChildWidget(ToolBarButton, TEXT("SLayeredImage"));
+								TSharedPtr<SWidget> TextBlockWidget = FindFirstChildWidget(ToolBarButton, TEXT("STextBlock"));
+								if (LayeredImageWidget.IsValid() && TextBlockWidget.IsValid())
+								{
+									TSharedPtr<SImage> ImageWidget = StaticCastSharedPtr<SImage>(LayeredImageWidget);
+									TSharedPtr<STextBlock> TextWidget = StaticCastSharedPtr<STextBlock>(TextBlockWidget);
+									// Check if this Section.Tool has an explicit color entry. If not, fallback
+									// to any Section-wide color entry, otherwise leave the tint alone.
+									FLinearColor* TintColor = SectionToolIconColors ? SectionToolIconColors->Find(TextWidget->GetText().ToString()) : nullptr;
+									if (!TintColor)
+									{
+										const FString* SourceText = FTextInspector::GetSourceString(TextWidget->GetText());
+										TintColor = SectionToolIconColors && SourceText ? SectionToolIconColors->Find(*SourceText) : nullptr;
+										if (!TintColor)
+										{
+											TintColor = SectionIconColor;
+										}
+									}
+									if (TintColor)
+									{
+										ImageWidget->SetColorAndOpacity(FSlateColor(*TintColor));
+									}
+								}
+							}
+						}
+					}
+				}
+			}
+		}
 	}
 }
 
diff --git a/Engine/Plugins/Experimental/ChaosEditor/Source/FractureEditor/Public/FractureEditorMode.h b/Engine/Plugins/Experimental/ChaosEditor/Source/FractureEditor/Public/FractureEditorMode.h
--- a/Engine/Plugins/Experimental/ChaosEditor/Source/FractureEditor/Public/FractureEditorMode.h
+++ b/Engine/Plugins/Experimental/ChaosEditor/Source/FractureEditor/Public/FractureEditorMode.h
@@ -94,6 +94,29 @@ struct FFractureModeCustomSectionColor
 };
 
 
+/**
+ * Defines a color to be used for a particular Tool Palette Section
+ */
+USTRUCT()
+struct FFractureModeCustomToolColor
+{
+	GENERATED_BODY()
+
+	/**
+	 * Name of Section or Tool in Fracture Mode Tool Palette
+	 *
+	 * Format:
+	 * SectionName        (Specifies a default color for all tools in the section.)
+	 * SectionName.ToolName        (Specifies an override color for a specific tool in the given section.)
+	 */
+	UPROPERTY(EditAnywhere, Category = "ToolColor")
+	FString ToolName = TEXT("");
+
+	/** Custom Tool Color */
+	UPROPERTY(EditAnywhere, Category = "ToolColor")
+	FLinearColor Color = FLinearColor::Gray;
+};
+
 
 UCLASS(config = Editor)
 class FRACTUREEDITOR_API UFractureModeCustomizationSettings : public UDeveloperSettings
@@ -131,5 +154,15 @@ public:
 	/** Custom Section Header Colors for listed Sections in the Fracture Mode Tool Palette */
 	UPROPERTY(config, EditAnywhere, Category = "Fracture Mode|UI Customization")
 	TArray<FFractureModeCustomSectionColor> SectionColors;
+
+	/**
+	 * Custom Tool Colors for listed Tools in the Fracture Mode Tool Palette.
+	 * 
+	 * Format:
+	 * SectionName        (Specifies a default color for all tools in the section.)
+	 * SectionName.ToolName        (Specifies an override color for a specific tool in the given section.)
+	 */
+	UPROPERTY(config, EditAnywhere, Category = "Fracture Mode|UI Customization")
+	TArray<FFractureModeCustomToolColor> ToolColors;
 };
 
