ModelingTools: Add support for "exclusive" tools that require accept/cancel by disabling other tools in the tool palette while exclusive tools are active.

#rnx

diff --git a/Engine/Plugins/Editor/ModelingToolsEditorMode/Source/ModelingToolsEditorMode/Private/ModelingToolsEditorMode.cpp b/Engine/Plugins/Editor/ModelingToolsEditorMode/Source/ModelingToolsEditorMode/Private/ModelingToolsEditorMode.cpp
--- a/Engine/Plugins/Editor/ModelingToolsEditorMode/Source/ModelingToolsEditorMode/Private/ModelingToolsEditorMode.cpp
+++ b/Engine/Plugins/Editor/ModelingToolsEditorMode/Source/ModelingToolsEditorMode/Private/ModelingToolsEditorMode.cpp
@@ -745,6 +745,15 @@ void UModelingToolsEditorMode::Exit()
 		GetModeManager()->GetSelectedActors()->GetElementSelectionSet()->OnPreChange().Remove(SelectionModifiedEventHandle);
 	}
 
+	// exit any exclusive active tools w/ cancel
+	if (UInteractiveTool* ActiveTool = GetToolManager()->GetActiveTool(EToolSide::Left))
+	{
+		if (Cast<IInteractiveToolExclusiveToolAPI>(ActiveTool))
+		{
+			GetToolManager()->DeactivateTool(EToolSide::Left, EToolShutdownType::Cancel);
+		}
+	}
+
 	//
 	// Engine Analytics
 	//
@@ -792,6 +801,23 @@ void UModelingToolsEditorMode::Exit()
 	UEdMode::Exit();
 }
 
+bool UModelingToolsEditorMode::ShouldToolStartBeAllowed(const FString& ToolIdentifier) const
+{
+	if (UInteractiveToolManager* Manager = GetToolManager())
+	{
+		if (UInteractiveTool* Tool = Manager->GetActiveTool(EToolSide::Left))
+		{
+			IInteractiveToolExclusiveToolAPI* ExclusiveAPI = Cast<IInteractiveToolExclusiveToolAPI>(Tool);
+			if (ExclusiveAPI)
+			{
+				return false;
+			}
+		}
+	}
+	return Super::ShouldToolStartBeAllowed(ToolIdentifier);
+}
+
+
 void UModelingToolsEditorMode::CreateToolkit()
 {
 	Toolkit = MakeShareable(new FModelingToolsEditorModeToolkit);
diff --git a/Engine/Plugins/Editor/ModelingToolsEditorMode/Source/ModelingToolsEditorMode/Public/ModelingToolsEditorMode.h b/Engine/Plugins/Editor/ModelingToolsEditorMode/Source/ModelingToolsEditorMode/Public/ModelingToolsEditorMode.h
--- a/Engine/Plugins/Editor/ModelingToolsEditorMode/Source/ModelingToolsEditorMode/Public/ModelingToolsEditorMode.h
+++ b/Engine/Plugins/Editor/ModelingToolsEditorMode/Source/ModelingToolsEditorMode/Public/ModelingToolsEditorMode.h
@@ -56,6 +56,9 @@ public:
 	// called when we "end" this editor mode (ie switch to another tab)
 	virtual void Exit() override;
 
+
+	virtual bool ShouldToolStartBeAllowed(const FString& ToolIdentifier) const override;
+
 	//////////////////
 	// End of UEdMode interface
 	//////////////////
diff --git a/Engine/Source/Runtime/InteractiveToolsFramework/Public/InteractiveToolQueryInterfaces.h b/Engine/Source/Runtime/InteractiveToolsFramework/Public/InteractiveToolQueryInterfaces.h
--- a/Engine/Source/Runtime/InteractiveToolsFramework/Public/InteractiveToolQueryInterfaces.h
+++ b/Engine/Source/Runtime/InteractiveToolsFramework/Public/InteractiveToolQueryInterfaces.h
@@ -127,4 +127,27 @@ public:
 	virtual bool ExecuteNestedAcceptCommand() { return false; }
 
 
-};
\ No newline at end of file
+};
+
+
+// UInterface for IInteractiveToolExclusiveToolAPI
+UINTERFACE(MinimalAPI)
+class UInteractiveToolExclusiveToolAPI : public UInterface
+{
+	GENERATED_BODY()
+};
+
+/**
+ * IInteractiveToolExclusiveToolAPI provides an API to inform the
+ * ToolManager about tool exclusivity. An exclusive tool prevents other
+ * tools from building & activating while the tool is active. This is
+ * useful in scenarios where tools want to enforce an explicit Accept,
+ * Cancel or Complete user input to exit the tool.
+ */
+class IInteractiveToolExclusiveToolAPI
+{
+	GENERATED_BODY()
+};
+
+
+
