ModelingTools: settings panel UI updates for VertexSculpt Tool
- add Details Customizations for FBrushToolRadius Property, USculptBrushProperties, UVertexBrushSculptProperties, UVertexBrushAlphaProperties tool property sets
- add tracking and unregistration of customizations in ModelingToolsEditorMode
- add temporary icons for sculpt brushes and falloffs, png files and style
#rnx

diff --git a/Engine/Plugins/Editor/ModelingToolsEditorMode/Source/ModelingToolsEditorMode/ModelingToolsEditorMode.Build.cs b/Engine/Plugins/Editor/ModelingToolsEditorMode/Source/ModelingToolsEditorMode/ModelingToolsEditorMode.Build.cs
--- a/Engine/Plugins/Editor/ModelingToolsEditorMode/Source/ModelingToolsEditorMode/ModelingToolsEditorMode.Build.cs
+++ b/Engine/Plugins/Editor/ModelingToolsEditorMode/Source/ModelingToolsEditorMode/ModelingToolsEditorMode.Build.cs
@@ -64,6 +64,7 @@ public class ModelingToolsEditorMode : ModuleRules
 				"MeshLODToolset",
 				"ToolWidgets",
 				"EditorWidgets",
+				"ModelingEditorUI",
 				"StylusInput",
 				"DeveloperSettings"
 
diff --git a/Engine/Plugins/Editor/ModelingToolsEditorMode/Source/ModelingToolsEditorMode/Private/DetailsCustomizations/MeshVertexSculptToolCustomizations.cpp b/Engine/Plugins/Editor/ModelingToolsEditorMode/Source/ModelingToolsEditorMode/Private/DetailsCustomizations/MeshVertexSculptToolCustomizations.cpp
new file mode 100644
--- /dev/null
+++ b/Engine/Plugins/Editor/ModelingToolsEditorMode/Source/ModelingToolsEditorMode/Private/DetailsCustomizations/MeshVertexSculptToolCustomizations.cpp
@@ -0,0 +1,531 @@
+// Copyright Epic Games, Inc. All Rights Reserved.
+
+#include "DetailsCustomizations/MeshVertexSculptToolCustomizations.h"
+#include "DetailCategoryBuilder.h"
+#include "DetailLayoutBuilder.h"
+#include "DetailWidgetRow.h"
+#include "PropertyCustomizationHelpers.h"
+#include "Widgets/Input/SButton.h"
+#include "Widgets/Input/SCheckBox.h"
+#include "Widgets/Text/STextBlock.h"
+#include "Widgets/Input/SComboButton.h"
+#include "Widgets/Views/STileView.h"
+#include "Widgets/Images/SImage.h"
+#include "IDetailChildrenBuilder.h"
+#include "Internationalization/BreakIterator.h"
+#include "Framework/MultiBox/MultiBoxBuilder.h"
+#include "IContentBrowserSingleton.h"
+#include "ContentBrowserModule.h"
+#include "SAssetView.h"
+#include "SAssetDropTarget.h"
+
+#include "ModelingWidgets/SComboPanel.h"
+#include "ModelingWidgets/SDynamicNumericEntry.h"
+#include "ModelingWidgets/SToolInputAssetPicker.h"
+#include "ModelingWidgets/SToolInputAssetComboPanel.h"
+#include "ModelingWidgets/ModelingCustomizationUtil.h"
+#include "DetailsCustomizations/ModelingToolsBrushSizeCustomization.h"
+
+#include "Styling/SlateStyleMacros.h"
+#include "ModelingToolsEditorModeStyle.h"
+#include "ModelingToolsEditorModeSettings.h"
+
+#include "MeshVertexSculptTool.h"
+
+using namespace UE::ModelingUI;
+
+#define LOCTEXT_NAMESPACE "MeshVertexSculptToolCustomizations"
+
+TSharedRef<IDetailCustomization> FSculptBrushPropertiesDetails::MakeInstance()
+{
+	return MakeShareable(new FSculptBrushPropertiesDetails);
+}
+
+
+void FSculptBrushPropertiesDetails::CustomizeDetails(IDetailLayoutBuilder& DetailBuilder)
+{
+	TSharedPtr<IPropertyHandle> FlowRateHandle = DetailBuilder.GetProperty(GET_MEMBER_NAME_CHECKED(USculptBrushProperties, FlowRate), USculptBrushProperties::StaticClass());
+	ensure(FlowRateHandle->IsValidHandle());
+
+	TSharedPtr<IPropertyHandle> SpacingHandle = DetailBuilder.GetProperty(GET_MEMBER_NAME_CHECKED(USculptBrushProperties, Spacing), USculptBrushProperties::StaticClass());
+	ensure(SpacingHandle->IsValidHandle());
+	SpacingHandle->MarkHiddenByCustomization();
+
+	TSharedPtr<SDynamicNumericEntry::FDataSource> FlowRateSource = SDynamicNumericEntry::MakeSimpleDataSource(
+		FlowRateHandle, TInterval<float>(0.0f, 1.0f), TInterval<float>(0.0f, 1.0f));
+	TSharedPtr<SDynamicNumericEntry::FDataSource> SpacingSource = SDynamicNumericEntry::MakeSimpleDataSource(
+		SpacingHandle, TInterval<float>(0.0f, 1000.0f), TInterval<float>(0.0f, 4.0f));
+
+	DetailBuilder.EditDefaultProperty(FlowRateHandle)->CustomWidget()
+		.OverrideResetToDefault(FResetToDefaultOverride::Hide())
+		.WholeRowContent()
+	[
+		MakeTwoWidgetDetailRowHBox(
+			MakeFixedWidthLabelSliderHBox(FlowRateHandle, FlowRateSource, FSculptToolsUIConstants::SculptShortLabelWidth),
+			MakeFixedWidthLabelSliderHBox(SpacingHandle, SpacingSource, FSculptToolsUIConstants::SculptShortLabelWidth)
+		)
+	];
+
+	TSharedPtr<IPropertyHandle> LazynessHandle = DetailBuilder.GetProperty(GET_MEMBER_NAME_CHECKED(USculptBrushProperties, Lazyness), USculptBrushProperties::StaticClass());
+	ensure(LazynessHandle->IsValidHandle());
+
+	TSharedPtr<IPropertyHandle> HitBackFacesHandle = DetailBuilder.GetProperty(GET_MEMBER_NAME_CHECKED(USculptBrushProperties, bHitBackFaces), USculptBrushProperties::StaticClass());
+	ensure(HitBackFacesHandle->IsValidHandle());
+	HitBackFacesHandle->MarkHiddenByCustomization();
+
+	// todo: 0-100 mapping
+	TSharedPtr<SDynamicNumericEntry::FDataSource> LazynessSource = SDynamicNumericEntry::MakeSimpleDataSource(
+		LazynessHandle, TInterval<float>(0.0f, 1.0f), TInterval<float>(0.0f, 1.0f));
+
+	DetailBuilder.EditDefaultProperty(LazynessHandle)->CustomWidget()
+		.OverrideResetToDefault(FResetToDefaultOverride::Hide())
+		.WholeRowContent()
+	[
+		MakeTwoWidgetDetailRowHBox(
+			MakeFixedWidthLabelSliderHBox(LazynessHandle, LazynessSource, FSculptToolsUIConstants::SculptShortLabelWidth),
+			MakeBoolToggleButton(HitBackFacesHandle, LOCTEXT("HitBackFacesText", "Hit Back Faces") )
+		)
+	];
+
+}
+
+
+TSharedRef<IDetailCustomization> FVertexBrushSculptPropertiesDetails::MakeInstance()
+{
+	return MakeShareable(new FVertexBrushSculptPropertiesDetails);
+}
+
+void FVertexBrushSculptPropertiesDetails::CustomizeDetails(IDetailLayoutBuilder& DetailBuilder)
+{
+	TSharedPtr<IPropertyHandle> BrushTypeHandle = DetailBuilder.GetProperty(GET_MEMBER_NAME_CHECKED(UVertexBrushSculptProperties, PrimaryBrushType), UVertexBrushSculptProperties::StaticClass());
+	ensure(BrushTypeHandle->IsValidHandle());
+
+	TArray<TWeakObjectPtr<UObject>> ObjectsBeingCustomized;
+	DetailBuilder.GetObjectsBeingCustomized(ObjectsBeingCustomized);
+	check(ObjectsBeingCustomized.Num() > 0);
+	UVertexBrushSculptProperties* BrushProperties = CastChecked<UVertexBrushSculptProperties>(ObjectsBeingCustomized[0]);
+	UMeshVertexSculptTool* Tool = BrushProperties->Tool.Get();
+	TargetTool = Tool;
+
+	const TArray<UMeshSculptToolBase::FBrushTypeInfo>& BrushTypeInfos = Tool->GetRegisteredPrimaryBrushTypes();
+	int32 CurrentBrushType = (int32)BrushProperties->PrimaryBrushType;
+	int32 CurrentBrushTypeIndex = 0;
+
+	TArray<TSharedPtr<SComboPanel::FComboPanelItem>> BrushTypeItems;
+	for (const UMeshSculptToolBase::FBrushTypeInfo& BrushTypeInfo : BrushTypeInfos)
+	{
+		TSharedPtr<SComboPanel::FComboPanelItem> NewBrushTypeItem = MakeShared<SComboPanel::FComboPanelItem>();
+		NewBrushTypeItem->Name = BrushTypeInfo.Name;
+		NewBrushTypeItem->Identifier = BrushTypeInfo.Identifier;
+		const FString* SourceBrushName = FTextInspector::GetSourceString(BrushTypeInfo.Name);
+		NewBrushTypeItem->Icon = FModelingToolsEditorModeStyle::Get()->GetBrush( FName(TEXT("BrushTypeIcons.") + (*SourceBrushName) ));
+		if (NewBrushTypeItem->Identifier == CurrentBrushType)
+		{
+			CurrentBrushTypeIndex = BrushTypeItems.Num();
+		}
+		BrushTypeItems.Add(NewBrushTypeItem);
+	}
+
+	float ComboIconSize = 60;
+	float FlyoutIconSize = 100;
+	float FlyoutWidth = 840;
+
+	TSharedPtr<SComboPanel> BrushTypeCombo = SNew(SComboPanel)
+		.ToolTipText(BrushTypeHandle->GetToolTipText())
+		.ComboButtonTileSize(FVector2D(ComboIconSize, ComboIconSize))
+		.FlyoutTileSize(FVector2D(FlyoutIconSize, FlyoutIconSize))
+		.FlyoutSize(FVector2D(FlyoutWidth, 1))
+		.ListItems(BrushTypeItems)
+		.OnSelectionChanged_Lambda([this](TSharedPtr<SComboPanel::FComboPanelItem> NewSelectedItem) {
+			if ( TargetTool.IsValid() )
+			{
+				TargetTool.Get()->SetActiveBrushType(NewSelectedItem->Identifier);
+			}
+		})
+		.FlyoutHeaderText(LOCTEXT("BrushesHeader", "Brush Types"))
+		.InitialSelectionIndex(CurrentBrushTypeIndex);
+
+
+	TSharedPtr<IPropertyHandle> FalloffTypeHandle = DetailBuilder.GetProperty(GET_MEMBER_NAME_CHECKED(UVertexBrushSculptProperties, PrimaryFalloffType), UVertexBrushSculptProperties::StaticClass());
+	ensure(FalloffTypeHandle->IsValidHandle());
+
+	const TArray<UMeshSculptToolBase::FFalloffTypeInfo>& FalloffTypeInfos = Tool->GetRegisteredPrimaryFalloffTypes();
+	int32 CurrentFalloffType = (int32)BrushProperties->PrimaryFalloffType;
+	int32 CurrentFalloffTypeIndex = 0;
+
+	TArray<TSharedPtr<SComboPanel::FComboPanelItem>> FalloffTypeItems;
+	for (const UMeshSculptToolBase::FFalloffTypeInfo& FalloffTypeInfo : FalloffTypeInfos)
+	{
+		TSharedPtr<SComboPanel::FComboPanelItem> NewFalloffTypeItem = MakeShared<SComboPanel::FComboPanelItem>();
+		NewFalloffTypeItem->Name = FalloffTypeInfo.Name;
+		NewFalloffTypeItem->Identifier = FalloffTypeInfo.Identifier;
+		NewFalloffTypeItem->Icon = FModelingToolsEditorModeStyle::Get()->GetBrush( FName(TEXT("BrushFalloffIcons.") + FalloffTypeInfo.StringIdentifier) );
+		if (NewFalloffTypeItem->Identifier == CurrentFalloffType)
+		{
+			CurrentFalloffTypeIndex = FalloffTypeItems.Num();
+		}
+		FalloffTypeItems.Add(NewFalloffTypeItem);
+	}
+
+	TSharedPtr<SComboPanel> FalloffTypeCombo = SNew(SComboPanel)
+		.ToolTipText(FalloffTypeHandle->GetToolTipText())
+		.ComboButtonTileSize(FVector2D(18, 18))
+		.FlyoutTileSize(FVector2D(FlyoutIconSize, FlyoutIconSize))
+		.FlyoutSize(FVector2D(FlyoutWidth, 1))
+		.ListItems(FalloffTypeItems)
+		.ComboDisplayType(SComboPanel::EComboDisplayType::IconAndLabel)
+		.OnSelectionChanged_Lambda([this](TSharedPtr<SComboPanel::FComboPanelItem> NewSelectedFalloffTypeItem) {
+			if ( TargetTool.IsValid() )
+			{
+				TargetTool.Get()->SetActiveFalloffType(NewSelectedFalloffTypeItem->Identifier);
+			}
+		})
+		.FlyoutHeaderText(LOCTEXT("FalloffsHeader", "Falloff Types"))
+		.InitialSelectionIndex(CurrentFalloffTypeIndex);
+
+	//DetailBuilder.HideProperty(FalloffTypeHandle);
+	FalloffTypeHandle->MarkHiddenByCustomization();
+
+
+	TSharedPtr<IPropertyHandle> BrushFilterHandle = DetailBuilder.GetProperty(GET_MEMBER_NAME_CHECKED(UVertexBrushSculptProperties, BrushFilter), UVertexBrushSculptProperties::StaticClass());
+	BrushFilterHandle->MarkHiddenByCustomization();
+
+	TSharedPtr<IPropertyHandle> FreezeTargetHandle = DetailBuilder.GetProperty(GET_MEMBER_NAME_CHECKED(UVertexBrushSculptProperties, bFreezeTarget), UVertexBrushSculptProperties::StaticClass());
+	FreezeTargetHandle->MarkHiddenByCustomization();
+
+
+	DetailBuilder.EditDefaultProperty(BrushTypeHandle)->CustomWidget()
+		.OverrideResetToDefault(FResetToDefaultOverride::Hide())
+		.WholeRowContent()
+		[
+			SNew(SHorizontalBox)
+			+ SHorizontalBox::Slot()
+			.Padding(FMargin(0, ModelingUIConstants::DetailRowVertPadding))
+			.AutoWidth()
+			[
+				SNew(SBox)
+				.HeightOverride(ComboIconSize+14)
+				[
+					BrushTypeCombo->AsShared()
+				]
+			]
+			+ SHorizontalBox::Slot()
+			.Padding(FMargin(ModelingUIConstants::MultiWidgetRowHorzPadding,ModelingUIConstants::DetailRowVertPadding,0,ModelingUIConstants::MultiWidgetRowHorzPadding))
+			.FillWidth(1.0f)
+			[
+				SNew(SVerticalBox)
+					+ SVerticalBox::Slot()
+					.Padding(FMargin(0))
+					.AutoHeight()
+					[
+						SNew(SHorizontalBox)
+							+ SHorizontalBox::Slot()
+							.AutoWidth()
+							[
+								WrapInFixedWidthBox(FalloffTypeHandle->CreatePropertyNameWidget(), FSculptToolsUIConstants::SculptShortLabelWidth)
+							]
+							+ SHorizontalBox::Slot()
+							.FillWidth(1.0f)
+							[
+								FalloffTypeCombo->AsShared()
+							]
+					]
+
+					+ SVerticalBox::Slot()
+					.Padding(FMargin(0, ModelingUIConstants::DetailRowVertPadding, 0, 0))
+					.AutoHeight()
+					[
+						SNew(SHorizontalBox)
+						+ SHorizontalBox::Slot()
+						.AutoWidth()
+						[
+							WrapInFixedWidthBox(BrushFilterHandle->CreatePropertyNameWidget(), FSculptToolsUIConstants::SculptShortLabelWidth)
+						]
+						+ SHorizontalBox::Slot()
+						.FillWidth(1.0f)
+						[
+							MakeRegionFilterWidget()->AsShared()
+						]
+					]
+
+					+ SVerticalBox::Slot()
+					.Padding(FMargin(0, ModelingUIConstants::DetailRowVertPadding, 0, 0))
+					.AutoHeight()
+					[
+						MakeFreezeTargetWidget()->AsShared()
+					]
+			]
+		];
+
+}
+
+
+TSharedPtr<SWidget> FVertexBrushSculptPropertiesDetails::MakeRegionFilterWidget()
+{
+	static const FText RegionFilterLabels[] = {
+		LOCTEXT("RegionFilterNone", "Vol"),  
+		LOCTEXT("RegionFilterComponent", "Cmp"),  
+		LOCTEXT("RegionFilterPolyGroup", "Grp")
+	};
+	static const FText RegionFilterTooltips[] = {
+		LOCTEXT("RegionFilterNoneTooltip", "Do not filter brush area, include all triangles in brush sphere"),
+		LOCTEXT("RegionFilterComponentTooltip", "Only apply brush to triangles in the same connected mesh component/island"),
+		LOCTEXT("RegionFilterPolygroupTooltip", "Only apply brush to triangles with the same PolyGroup"),
+	};
+
+	auto MakeRegionFilterButton = [this](EMeshVertexSculptBrushFilterType FilterType)
+	{
+		return SNew(SCheckBox)
+			.Style(FAppStyle::Get(), "DetailsView.SectionButton")
+			.Padding(FMargin(2, 2))
+			.HAlign(HAlign_Center)
+			.ToolTipText( RegionFilterTooltips[(int32)FilterType] )
+			.OnCheckStateChanged_Lambda([this, FilterType](ECheckBoxState State) {
+				if (TargetTool.IsValid() && State == ECheckBoxState::Checked)
+				{
+					TargetTool.Get()->SetRegionFilterType( static_cast<int32>(FilterType) );
+				}
+			})
+			.IsChecked_Lambda([this, FilterType]() {
+				return (TargetTool.IsValid() && TargetTool.Get()->SculptProperties->BrushFilter == FilterType) ? ECheckBoxState::Checked : ECheckBoxState::Unchecked;
+			})
+			[
+				SNew(SHorizontalBox)
+				+ SHorizontalBox::Slot()
+				.VAlign(VAlign_Center)
+				.Padding(FMargin(0))
+				.AutoWidth()
+				[
+					SNew(STextBlock)
+					.Justification(ETextJustify::Center)
+					.TextStyle(FAppStyle::Get(), "DetailsView.CategoryTextStyle")
+					.Text( RegionFilterLabels[(int32)FilterType] )
+				]
+			];
+	};
+
+
+	return SNew(SHorizontalBox)
+		+ SHorizontalBox::Slot()
+		.FillWidth(1.0f)
+		[
+			MakeRegionFilterButton(EMeshVertexSculptBrushFilterType::None)
+		]
+		+ SHorizontalBox::Slot()
+		.FillWidth(1.0f)
+		[
+			MakeRegionFilterButton(EMeshVertexSculptBrushFilterType::Component)
+		]
+		+ SHorizontalBox::Slot()
+		.FillWidth(1.0f)
+		[
+			MakeRegionFilterButton(EMeshVertexSculptBrushFilterType::PolyGroup)
+		];
+}
+
+
+TSharedPtr<SWidget> FVertexBrushSculptPropertiesDetails::MakeFreezeTargetWidget()
+{
+	return SNew(SCheckBox)
+		.Style(FAppStyle::Get(), "DetailsView.SectionButton")
+		.Padding(FMargin(0, 4))
+		.ToolTipText(LOCTEXT("FreezeTargetTooltip", "When Freeze Target is toggled on, the Brush Target Surface will be Frozen in its current state, until toggled off. Brush strokes will be applied relative to the Target Surface, for applicable Brushes"))
+		.HAlign(HAlign_Center)
+		.OnCheckStateChanged(this, &FVertexBrushSculptPropertiesDetails::OnSetFreezeTarget)
+		.IsChecked_Lambda([this] { return IsFreezeTargetEnabled() ? ECheckBoxState::Checked : ECheckBoxState::Unchecked; } )
+		[
+			SNew(SHorizontalBox)
+			+ SHorizontalBox::Slot()
+			.VAlign(VAlign_Center)
+			.Padding(FMargin(0))
+			.AutoWidth()
+			[
+				SNew(STextBlock)
+				.Justification(ETextJustify::Center)
+				.TextStyle(FAppStyle::Get(), "DetailsView.CategoryTextStyle")
+				.Text_Lambda( [this] { return IsFreezeTargetEnabled() ? LOCTEXT("FreezeTarget", "UnFreeze Target") : LOCTEXT("FreezeTarget", "Freeze Target"); } )
+			]
+		];
+}
+
+
+
+
+FReply FVertexBrushSculptPropertiesDetails::OnToggledFreezeTarget()
+{
+	if ( TargetTool.IsValid() )
+	{
+		TargetTool.Get()->SculptProperties->bFreezeTarget = !TargetTool.Get()->SculptProperties->bFreezeTarget;
+	}
+	return FReply::Handled();
+}
+
+void FVertexBrushSculptPropertiesDetails::OnSetFreezeTarget(ECheckBoxState State)
+{
+	if ( TargetTool.IsValid() )
+	{
+		TargetTool.Get()->SculptProperties->bFreezeTarget = (State == ECheckBoxState::Checked);
+	}
+}
+
+bool FVertexBrushSculptPropertiesDetails::IsFreezeTargetEnabled()
+{
+	return (TargetTool.IsValid()) ? TargetTool.Get()->SculptProperties->bFreezeTarget : false;
+}
+
+
+
+
+
+
+
+
+
+class FRecentAlphasProvider : public SToolInputAssetComboPanel::IRecentAssetsProvider
+{
+public:
+	TArray<FAssetData> RecentAssets;
+
+	virtual TArray<FAssetData> GetRecentAssetsList() override
+	{
+		return RecentAssets;
+	}
+	virtual void NotifyNewAsset(const FAssetData& NewAsset) override
+	{
+		if (NewAsset.GetAsset() == nullptr)
+		{
+			return;
+		}
+		for (int32 k = 0; k < RecentAssets.Num(); ++k)
+		{
+			if (RecentAssets[k] == NewAsset)
+			{
+				if (k == 0)
+				{
+					return;
+				}
+				RecentAssets.RemoveAt(k, 1, false);
+			}
+		}
+		RecentAssets.Insert(NewAsset, 0);
+		
+		if (RecentAssets.Num() > 10)
+		{
+			RecentAssets.SetNum(10, false);
+		}
+	}
+};
+
+
+
+TSharedRef<IDetailCustomization> FVertexBrushAlphaPropertiesDetails::MakeInstance()
+{
+	return MakeShareable(new FVertexBrushAlphaPropertiesDetails);
+}
+
+void FVertexBrushAlphaPropertiesDetails::CustomizeDetails(IDetailLayoutBuilder& DetailBuilder)
+{
+	// TODO: move this to a subsystem or UObject CDO
+	struct FRecentAlphasContainer
+	{
+		TSharedPtr<FRecentAlphasProvider> RecentAlphas;
+	};
+	static FRecentAlphasContainer RecentAlphasStatic;
+	if (RecentAlphasStatic.RecentAlphas.IsValid() == false)
+	{
+		RecentAlphasStatic.RecentAlphas = MakeShared<FRecentAlphasProvider>();
+	}
+	RecentAlphasProvider = RecentAlphasStatic.RecentAlphas;
+
+
+
+	TSharedPtr<IPropertyHandle> AlphaHandle = DetailBuilder.GetProperty(GET_MEMBER_NAME_CHECKED(UVertexBrushAlphaProperties, Alpha), UVertexBrushAlphaProperties::StaticClass());
+	ensure(AlphaHandle->IsValidHandle());
+
+	TSharedPtr<IPropertyHandle> RotationAngleHandle = DetailBuilder.GetProperty(GET_MEMBER_NAME_CHECKED(UVertexBrushAlphaProperties, RotationAngle), UVertexBrushAlphaProperties::StaticClass());
+	ensure(RotationAngleHandle->IsValidHandle());
+	RotationAngleHandle->MarkHiddenByCustomization();
+
+	TSharedPtr<IPropertyHandle> bRandomizeHandle = DetailBuilder.GetProperty(GET_MEMBER_NAME_CHECKED(UVertexBrushAlphaProperties, bRandomize), UVertexBrushAlphaProperties::StaticClass());
+	ensure(bRandomizeHandle->IsValidHandle());
+	bRandomizeHandle->MarkHiddenByCustomization();
+
+	TSharedPtr<IPropertyHandle> RandomRangeHandle = DetailBuilder.GetProperty(GET_MEMBER_NAME_CHECKED(UVertexBrushAlphaProperties, RandomRange), UVertexBrushAlphaProperties::StaticClass());
+	ensure(RandomRangeHandle->IsValidHandle());
+	RandomRangeHandle->MarkHiddenByCustomization();
+
+	TSharedPtr<SDynamicNumericEntry::FDataSource> RotationAngleSource = SDynamicNumericEntry::MakeSimpleDataSource(
+		RotationAngleHandle, TInterval<float>(-180.0f, 180.0f), TInterval<float>(-180.0f, 180.0f));
+	TSharedPtr<SDynamicNumericEntry::FDataSource> RandomRangeSource = SDynamicNumericEntry::MakeSimpleDataSource(
+		RandomRangeHandle, TInterval<float>(0.0f, 180.0f), TInterval<float>(0.0f, 180.0f));
+
+	float ComboIconSize = 60;
+
+	UModelingToolsModeCustomizationSettings* UISettings = GetMutableDefault<UModelingToolsModeCustomizationSettings>();
+	TArray<SToolInputAssetComboPanel::FNamedCollectionList> BrushAlphasLists;
+	for (const FModelingModeAssetCollectionSet& AlphasCollectionSet : UISettings->BrushAlphaSets)
+	{
+		SToolInputAssetComboPanel::FNamedCollectionList CollectionSet;
+		CollectionSet.Name = AlphasCollectionSet.Name;
+		for (FCollectionReference CollectionRef : AlphasCollectionSet.Collections)
+		{
+			CollectionSet.CollectionNames.Add(FCollectionNameType(CollectionRef.CollectionName, ECollectionShareType::CST_Local));
+		}
+		BrushAlphasLists.Add(CollectionSet);
+	}
+
+	TSharedPtr<SToolInputAssetComboPanel> AlphaAssetPicker = SNew(SToolInputAssetComboPanel)
+		.AssetClassType(UTexture2D::StaticClass())		// can infer from property...
+		.Property(AlphaHandle)
+		.ComboButtonTileSize(FVector2D(ComboIconSize, ComboIconSize))
+		.FlyoutTileSize(FVector2D(80, 80))
+		.FlyoutSize(FVector2D(1000, 600))
+		.RecentAssetsProvider(RecentAlphasProvider)
+		.CollectionSets(BrushAlphasLists);
+
+	DetailBuilder.EditDefaultProperty(AlphaHandle)->CustomWidget()
+		.OverrideResetToDefault(FResetToDefaultOverride::Hide())
+		.WholeRowContent()
+		[
+			SNew(SHorizontalBox)
+			+ SHorizontalBox::Slot()
+			.Padding(FMargin(0, ModelingUIConstants::DetailRowVertPadding))
+			.AutoWidth()
+			[
+				SNew(SBox)
+				.HeightOverride(ComboIconSize+14)
+				[
+					AlphaAssetPicker->AsShared()
+				]
+			]
+			+ SHorizontalBox::Slot()
+			.Padding(FMargin(ModelingUIConstants::MultiWidgetRowHorzPadding,ModelingUIConstants::DetailRowVertPadding,0,ModelingUIConstants::MultiWidgetRowHorzPadding))
+			.FillWidth(1.0f)
+			[
+				SNew(SVerticalBox)
+					+ SVerticalBox::Slot()
+					.Padding(FMargin(0))
+					.AutoHeight()
+					[
+						MakeFixedWidthLabelSliderHBox(RotationAngleHandle, RotationAngleSource, FSculptToolsUIConstants::SculptShortLabelWidth)
+					]
+
+					+ SVerticalBox::Slot()
+					.Padding(FMargin(0, ModelingUIConstants::DetailRowVertPadding))
+					.AutoHeight()
+					[
+						MakeToggleSliderHBox(bRandomizeHandle, LOCTEXT("RandomizeLabel", "Rand"), RandomRangeSource, FSculptToolsUIConstants::SculptShortLabelWidth)
+					]
+			]
+		];
+	
+
+
+}
+
+
+
+
+
+
+#undef LOCTEXT_NAMESPACE
\ No newline at end of file
diff --git a/Engine/Plugins/Editor/ModelingToolsEditorMode/Source/ModelingToolsEditorMode/Private/DetailsCustomizations/ModelingToolsBrushSizeCustomization.cpp b/Engine/Plugins/Editor/ModelingToolsEditorMode/Source/ModelingToolsEditorMode/Private/DetailsCustomizations/ModelingToolsBrushSizeCustomization.cpp
new file mode 100644
--- /dev/null
+++ b/Engine/Plugins/Editor/ModelingToolsEditorMode/Source/ModelingToolsEditorMode/Private/DetailsCustomizations/ModelingToolsBrushSizeCustomization.cpp
@@ -0,0 +1,142 @@
+// Copyright Epic Games, Inc. All Rights Reserved.
+
+#include "DetailsCustomizations/ModelingToolsBrushSizeCustomization.h"
+
+#include "PropertyHandle.h"
+#include "IPropertyTypeCustomization.h"
+#include "IPropertyUtilities.h"
+#include "DetailWidgetRow.h"
+#include "EditorFontGlyphs.h"
+#include "Widgets/Input/SCheckBox.h"
+#include "ModelingWidgets/SDynamicNumericEntry.h"
+#include "ModelingWidgets/ModelingCustomizationUtil.h"
+
+#include "Sculpting/MeshSculptToolBase.h"
+
+using namespace UE::ModelingUI;
+
+#define LOCTEXT_NAMESPACE "ModelingToolsBrushSizeCustomization"
+
+
+
+TSharedRef<IPropertyTypeCustomization> FModelingToolsBrushSizeCustomization::MakeInstance()
+{
+	return MakeShareable(new FModelingToolsBrushSizeCustomization);
+}
+
+void FModelingToolsBrushSizeCustomization::CustomizeHeader(TSharedRef<IPropertyHandle> StructPropertyHandle, FDetailWidgetRow& HeaderRow, IPropertyTypeCustomizationUtils& StructCustomizationUtils)
+{
+	TSharedPtr<IPropertyHandle> SizeType = StructPropertyHandle->GetChildHandle(GET_MEMBER_NAME_CHECKED(FBrushToolRadius, SizeType));
+	SizeType->MarkHiddenByCustomization();
+
+	TSharedPtr<IPropertyHandle> AdaptiveSize = StructPropertyHandle->GetChildHandle(GET_MEMBER_NAME_CHECKED(FBrushToolRadius, AdaptiveSize));
+	AdaptiveSize->MarkHiddenByCustomization();
+
+	TSharedPtr<IPropertyHandle> WorldRadius = StructPropertyHandle->GetChildHandle(GET_MEMBER_NAME_CHECKED(FBrushToolRadius, WorldRadius));
+	WorldRadius->MarkHiddenByCustomization();
+
+	auto GetCurrentSizeType = [SizeType]() -> EBrushToolSizeType
+	{
+		uint8 SizeTypeInt = 0;
+		SizeType->GetValue(SizeTypeInt);
+		return (SizeTypeInt == (uint8)EBrushToolSizeType::World) ? EBrushToolSizeType::World : EBrushToolSizeType::Adaptive;
+	};
+	auto GetSizeHandle = [SizeType, AdaptiveSize, WorldRadius]() -> TSharedPtr<IPropertyHandle>
+	{
+		uint8 SizeTypeInt = 0;
+		SizeType->GetValue(SizeTypeInt);
+		return (SizeTypeInt == (uint8)EBrushToolSizeType::World) ? WorldRadius : AdaptiveSize;
+	};
+
+	TSharedPtr<SDynamicNumericEntry::FDataSource> NumericSource = MakeShared<SDynamicNumericEntry::FDataSource>();
+	NumericSource->SetValue = [GetSizeHandle](float NewSize, EPropertyValueSetFlags::Type Flags)
+	{
+		GetSizeHandle()->SetValue(NewSize, Flags);
+	};
+	NumericSource->GetValue = [GetSizeHandle]() -> float
+	{
+		float Size;
+		GetSizeHandle()->GetValue(Size);
+		return Size;
+	};
+	NumericSource->GetValueRange = [GetCurrentSizeType]() -> TInterval<float>
+	{
+		return (GetCurrentSizeType() == EBrushToolSizeType::Adaptive) ? TInterval<float>(0.0f, 10.0f) : TInterval<float>(0.01f, 50000.0f);
+	};
+	NumericSource->GetUIRange = [GetCurrentSizeType]() -> TInterval<float>
+	{
+		return (GetCurrentSizeType() == EBrushToolSizeType::Adaptive) ? TInterval<float>(0.0f, 1.0f) : TInterval<float>(1.0f, 1000.0f);
+	};
+
+	HeaderRow
+		.OverrideResetToDefault(FResetToDefaultOverride::Hide())
+		.WholeRowContent()
+	[
+		SNew(SHorizontalBox)
+
+		// name
+		+ SHorizontalBox::Slot()
+		.AutoWidth()
+		.Padding(FMargin(0,4))
+		[
+			WrapInFixedWidthBox(AdaptiveSize->CreatePropertyNameWidget(), FSculptToolsUIConstants::SculptShortLabelWidth)
+		]
+
+		// slider
+		+ SHorizontalBox::Slot()
+		.Padding(ModelingUIConstants::LabelWidgetMinPadding, 0,0,0)
+		.FillWidth(10.0f)
+		[
+			SNew(SDynamicNumericEntry)
+			.Source(NumericSource)
+		]
+
+		// toggle
+		+ SHorizontalBox::Slot()
+		.AutoWidth()
+		//.FillWidth(0.5f)
+		.Padding(FMargin(ModelingUIConstants::MultiWidgetRowHorzPadding, ModelingUIConstants::DetailRowVertPadding, 0, ModelingUIConstants::DetailRowVertPadding))
+		.HAlign(HAlign_Right)
+		[
+			SNew(SBox)
+			.WidthOverride(50)
+			[
+				SNew(SCheckBox)
+					.Style(FAppStyle::Get(), "DetailsView.SectionButton")
+					.Padding(FMargin(4, 2))
+					.HAlign(HAlign_Center)
+					.ToolTipText( LOCTEXT("WorldToggleTooltip", "Specify Brush Size in World Units") )
+					.OnCheckStateChanged_Lambda([this, SizeType](ECheckBoxState State) {
+						SizeType->SetValue( static_cast<uint8>( (State == ECheckBoxState::Checked) ? EBrushToolSizeType::World : EBrushToolSizeType::Adaptive ) );
+					})
+					.IsChecked_Lambda([this, SizeType]() {
+						uint8 CurValueInt;
+						SizeType->GetValue(CurValueInt);
+						return (CurValueInt == static_cast<uint8>(EBrushToolSizeType::World)) ? ECheckBoxState::Checked : ECheckBoxState::Unchecked;
+					})
+					[
+						SNew(SHorizontalBox)
+						+ SHorizontalBox::Slot()
+						.VAlign(VAlign_Center)
+						.Padding(FMargin(0))
+						.AutoWidth()
+						[
+							SNew(STextBlock)
+							.Justification(ETextJustify::Center)
+							.TextStyle(FAppStyle::Get(), "DetailsView.CategoryTextStyle")
+							.Text( LOCTEXT("World", "World") )
+						]
+					]
+			]
+		]
+	];
+
+}
+
+
+void FModelingToolsBrushSizeCustomization::CustomizeChildren(TSharedRef<IPropertyHandle> StructPropertyHandle, class IDetailChildrenBuilder& ChildBuilder, IPropertyTypeCustomizationUtils& StructCustomizationUtils)
+{
+}
+
+
+#undef LOCTEXT_NAMESPACE
\ No newline at end of file
diff --git a/Engine/Plugins/Editor/ModelingToolsEditorMode/Source/ModelingToolsEditorMode/Private/ModelingToolsEditorModeModule.cpp b/Engine/Plugins/Editor/ModelingToolsEditorMode/Source/ModelingToolsEditorMode/Private/ModelingToolsEditorModeModule.cpp
--- a/Engine/Plugins/Editor/ModelingToolsEditorMode/Source/ModelingToolsEditorMode/Private/ModelingToolsEditorModeModule.cpp
+++ b/Engine/Plugins/Editor/ModelingToolsEditorMode/Source/ModelingToolsEditorMode/Private/ModelingToolsEditorModeModule.cpp
@@ -12,6 +12,12 @@
 #include "Modules/ModuleManager.h"
 #include "PropertyEditorModule.h"
 #include "DetailsCustomizations/ModelingToolPropertyCustomizations.h"
+#include "DetailsCustomizations/ModelingToolsBrushSizeCustomization.h"
+#include "DetailsCustomizations/MeshVertexSculptToolCustomizations.h"
+
+#include "PropertySets/AxisFilterPropertyType.h"
+#include "MeshVertexSculptTool.h"
+
 
 #define LOCTEXT_NAMESPACE "FModelingToolsEditorModeModule"
 
@@ -28,6 +34,20 @@ void FModelingToolsEditorModeModule::ShutdownModule()
 	FModelingToolsManagerCommands::Unregister();
 	FModelingModeActionCommands::Unregister();
 
+	// Unregister customizations
+	FPropertyEditorModule* PropertyEditorModule = FModuleManager::GetModulePtr<FPropertyEditorModule>("PropertyEditor");
+	if (PropertyEditorModule)
+	{
+		for (FName ClassName : ClassesToUnregisterOnShutdown)
+		{
+			PropertyEditorModule->UnregisterCustomClassLayout(ClassName);
+		}
+		for (FName PropertyName : PropertiesToUnregisterOnShutdown)
+		{
+			PropertyEditorModule->UnregisterCustomPropertyTypeLayout(PropertyName);
+		}
+	}
+
 	// Unregister slate style overrides
 	FModelingToolsEditorModeStyle::Shutdown();
 }
@@ -42,9 +62,23 @@ void FModelingToolsEditorModeModule::OnPostEngineInit()
 	FModelingToolsManagerCommands::Register();
 	FModelingModeActionCommands::Register();
 
+	// same as ClassesToUnregisterOnShutdown but for properties, there is none right now
+	PropertiesToUnregisterOnShutdown.Reset();
+	ClassesToUnregisterOnShutdown.Reset();
+
+
 	// Register details view customizations
 	FPropertyEditorModule& PropertyModule = FModuleManager::LoadModuleChecked<FPropertyEditorModule>("PropertyEditor");
 	PropertyModule.RegisterCustomPropertyTypeLayout("ModelingToolsAxisFilter", FOnGetPropertyTypeCustomizationInstance::CreateStatic(&FModelingToolsAxisFilterCustomization::MakeInstance));
+	PropertiesToUnregisterOnShutdown.Add(FModelingToolsAxisFilter::StaticStruct()->GetFName());
+	PropertyModule.RegisterCustomPropertyTypeLayout("BrushToolRadius", FOnGetPropertyTypeCustomizationInstance::CreateStatic(&FModelingToolsBrushSizeCustomization::MakeInstance));
+	PropertiesToUnregisterOnShutdown.Add(FBrushToolRadius::StaticStruct()->GetFName());
+	PropertyModule.RegisterCustomClassLayout("SculptBrushProperties", FOnGetDetailCustomizationInstance::CreateStatic(&FSculptBrushPropertiesDetails::MakeInstance));
+	ClassesToUnregisterOnShutdown.Add(USculptBrushProperties::StaticClass()->GetFName());
+	PropertyModule.RegisterCustomClassLayout("VertexBrushSculptProperties", FOnGetDetailCustomizationInstance::CreateStatic(&FVertexBrushSculptPropertiesDetails::MakeInstance));
+	ClassesToUnregisterOnShutdown.Add(UVertexBrushSculptProperties::StaticClass()->GetFName());
+	PropertyModule.RegisterCustomClassLayout("VertexBrushAlphaProperties", FOnGetDetailCustomizationInstance::CreateStatic(&FVertexBrushAlphaPropertiesDetails::MakeInstance));
+	ClassesToUnregisterOnShutdown.Add(UVertexBrushAlphaProperties::StaticClass()->GetFName());
 }
 
 #undef LOCTEXT_NAMESPACE
diff --git a/Engine/Plugins/Editor/ModelingToolsEditorMode/Source/ModelingToolsEditorMode/Private/ModelingToolsEditorModeStyle.cpp b/Engine/Plugins/Editor/ModelingToolsEditorMode/Source/ModelingToolsEditorMode/Private/ModelingToolsEditorModeStyle.cpp
--- a/Engine/Plugins/Editor/ModelingToolsEditorMode/Source/ModelingToolsEditorMode/Private/ModelingToolsEditorModeStyle.cpp
+++ b/Engine/Plugins/Editor/ModelingToolsEditorMode/Source/ModelingToolsEditorMode/Private/ModelingToolsEditorModeStyle.cpp
@@ -30,6 +30,12 @@ FName FModelingToolsEditorModeStyle::GetStyleSetName()
 	return ModelingToolsStyleName;
 }
 
+const FSlateBrush* FModelingToolsEditorModeStyle::GetBrush(FName PropertyName, const ANSICHAR* Specifier)
+{
+	return Get()->GetBrush(PropertyName, Specifier);
+}
+
+
 BEGIN_SLATE_FUNCTION_BUILD_OPTIMIZATION
 
 void FModelingToolsEditorModeStyle::Initialize()
@@ -40,6 +46,7 @@ void FModelingToolsEditorModeStyle::Initialize()
 	const FVector2D Icon20x20(20.0f, 20.0f);
 	const FVector2D Icon28x28(28.0f, 28.0f);
 	const FVector2D Icon40x40(40.0f, 40.0f);
+	const FVector2D Icon120(120.0f, 120.0f);
 
 	// Only register once
 	if (StyleSet.IsValid())
@@ -292,8 +299,46 @@ void FModelingToolsEditorModeStyle::Initialize()
 
 		StyleSet->Set("ModelingToolsManagerCommands.BeginSkinWeightsPaintTool", new IMAGE_BRUSH_SVG("Icons/SkinWeightsPaint", Icon20x20));
 		StyleSet->Set("ModelingToolsManagerCommands.BeginSkinWeightsPaintTool.Small", new IMAGE_BRUSH_SVG("Icons/SkinWeightsPaint", Icon20x20));
+
+
+		//
+		// Icons for brush falloffs in sculpt/etc tools
+		//
+
+		StyleSet->Set("BrushFalloffIcons.Smooth", new IMAGE_PLUGIN_BRUSH("Icons/BrushIcons/Falloff_Smooth", Icon120));
+		StyleSet->Set("BrushFalloffIcons.Linear", new IMAGE_PLUGIN_BRUSH("Icons/BrushIcons/Falloff_Linear", Icon120));
+		StyleSet->Set("BrushFalloffIcons.Inverse", new IMAGE_PLUGIN_BRUSH("Icons/BrushIcons/Falloff_Inverse", Icon120));
+		StyleSet->Set("BrushFalloffIcons.Round", new IMAGE_PLUGIN_BRUSH("Icons/BrushIcons/Falloff_Round", Icon120));
+		StyleSet->Set("BrushFalloffIcons.BoxSmooth", new IMAGE_PLUGIN_BRUSH("Icons/BrushIcons/Falloff_BoxSmooth", Icon120));
+		StyleSet->Set("BrushFalloffIcons.BoxLinear", new IMAGE_PLUGIN_BRUSH("Icons/BrushIcons/Falloff_BoxLinear", Icon120));
+		StyleSet->Set("BrushFalloffIcons.BoxInverse", new IMAGE_PLUGIN_BRUSH("Icons/BrushIcons/Falloff_BoxInverse", Icon120));
+		StyleSet->Set("BrushFalloffIcons.BoxRound", new IMAGE_PLUGIN_BRUSH("Icons/BrushIcons/Falloff_BoxRound", Icon120));
+
+
+		//
+		// Icons for brushes in sculpt/etc tools
+		//
+
+		StyleSet->Set("BrushTypeIcons.Smooth", new IMAGE_PLUGIN_BRUSH("Icons/BrushIcons/Brush_Smooth", Icon120));
+		StyleSet->Set("BrushTypeIcons.SmoothFill", new IMAGE_PLUGIN_BRUSH("Icons/BrushIcons/Brush_SmoothFill", Icon120));
+		StyleSet->Set("BrushTypeIcons.Move", new IMAGE_PLUGIN_BRUSH("Icons/BrushIcons/Brush_Move", Icon120));
+		StyleSet->Set("BrushTypeIcons.SculptN", new IMAGE_PLUGIN_BRUSH("Icons/BrushIcons/Brush_SculptN", Icon120));
+		StyleSet->Set("BrushTypeIcons.SculptV", new IMAGE_PLUGIN_BRUSH("Icons/BrushIcons/Brush_SculptV", Icon120));
+		StyleSet->Set("BrushTypeIcons.SculptMx", new IMAGE_PLUGIN_BRUSH("Icons/BrushIcons/Brush_SculptMx", Icon120));
+		StyleSet->Set("BrushTypeIcons.Inflate", new IMAGE_PLUGIN_BRUSH("Icons/BrushIcons/Brush_Inflate", Icon120));
+		StyleSet->Set("BrushTypeIcons.Pinch", new IMAGE_PLUGIN_BRUSH("Icons/BrushIcons/Brush_Pinch", Icon120));
+		StyleSet->Set("BrushTypeIcons.Flatten", new IMAGE_PLUGIN_BRUSH("Icons/BrushIcons/Brush_Flatten", Icon120));
+		StyleSet->Set("BrushTypeIcons.PlaneN", new IMAGE_PLUGIN_BRUSH("Icons/BrushIcons/Brush_PlaneN", Icon120));
+		StyleSet->Set("BrushTypeIcons.PlaneV", new IMAGE_PLUGIN_BRUSH("Icons/BrushIcons/Brush_PlaneV", Icon120));
+		StyleSet->Set("BrushTypeIcons.PlaneW", new IMAGE_PLUGIN_BRUSH("Icons/BrushIcons/Brush_PlaneW", Icon120));
+		StyleSet->Set("BrushTypeIcons.Scale", new IMAGE_PLUGIN_BRUSH("Icons/BrushIcons/Brush_Scale", Icon120));
+		StyleSet->Set("BrushTypeIcons.Grab", new IMAGE_PLUGIN_BRUSH("Icons/BrushIcons/Brush_Grab", Icon120));
+		StyleSet->Set("BrushTypeIcons.GrabSharp", new IMAGE_PLUGIN_BRUSH("Icons/BrushIcons/Brush_GrabSharp", Icon120));
+		StyleSet->Set("BrushTypeIcons.Twist", new IMAGE_PLUGIN_BRUSH("Icons/BrushIcons/Brush_Twist", Icon120));
+
 	}
 
+
 	FSlateStyleRegistry::RegisterSlateStyle(*StyleSet.Get());
 };
 
diff --git a/Engine/Plugins/Editor/ModelingToolsEditorMode/Source/ModelingToolsEditorMode/Public/DetailsCustomizations/MeshVertexSculptToolCustomizations.h b/Engine/Plugins/Editor/ModelingToolsEditorMode/Source/ModelingToolsEditorMode/Public/DetailsCustomizations/MeshVertexSculptToolCustomizations.h
new file mode 100644
--- /dev/null
+++ b/Engine/Plugins/Editor/ModelingToolsEditorMode/Source/ModelingToolsEditorMode/Public/DetailsCustomizations/MeshVertexSculptToolCustomizations.h
@@ -0,0 +1,61 @@
+// Copyright Epic Games, Inc. All Rights Reserved.
+
+#pragma once
+
+#include "IDetailCustomization.h"
+#include "Templates/SharedPointer.h"
+#include "IPropertyTypeCustomization.h"
+#include "Widgets/Input/SComboBox.h"
+#include "Widgets/Views/STileView.h"
+
+class IDetailLayoutBuilder;
+class IPropertyHandle;
+class IDetailChildrenBuilder;
+class SComboButton;
+class SBox;
+class UMeshVertexSculptTool;
+class FRecentAlphasProvider;
+
+
+// customization for USculptBrushProperties, creates two-column layout
+// for secondary brush properties like lazy/etc
+class FSculptBrushPropertiesDetails : public IDetailCustomization
+{
+public:
+	static TSharedRef<IDetailCustomization> MakeInstance();
+	void CustomizeDetails(IDetailLayoutBuilder& DetailBuilder) override;
+};
+
+
+// customization for vertexsculpt properties, creates combopanel for brush type,
+// small-style combopanel for falloff type, and stacks controls to the right
+class FVertexBrushSculptPropertiesDetails : public IDetailCustomization
+{
+public:
+	static TSharedRef<IDetailCustomization> MakeInstance();
+	void CustomizeDetails(IDetailLayoutBuilder& DetailBuilder) override;
+
+protected:
+	TWeakObjectPtr<UMeshVertexSculptTool> TargetTool;
+
+	TSharedPtr<SWidget> MakeRegionFilterWidget();
+	TSharedPtr<SWidget> MakeFreezeTargetWidget();
+	
+	TSharedPtr<SButton> FreezeTargetButton;
+	FReply OnToggledFreezeTarget();
+	void OnSetFreezeTarget(ECheckBoxState State);
+	bool IsFreezeTargetEnabled();
+};
+
+
+
+// customization for UVertexBrushAlphaProperties, creates custom asset picker
+// tileview-combopanel for brush alphas and stacks controls to the right
+class FVertexBrushAlphaPropertiesDetails : public IDetailCustomization
+{
+public:
+	static TSharedRef<IDetailCustomization> MakeInstance();
+	void CustomizeDetails(IDetailLayoutBuilder& DetailBuilder) override;
+
+	TSharedPtr<FRecentAlphasProvider> RecentAlphasProvider;
+};
\ No newline at end of file
diff --git a/Engine/Plugins/Editor/ModelingToolsEditorMode/Source/ModelingToolsEditorMode/Public/DetailsCustomizations/ModelingToolsBrushSizeCustomization.h b/Engine/Plugins/Editor/ModelingToolsEditorMode/Source/ModelingToolsEditorMode/Public/DetailsCustomizations/ModelingToolsBrushSizeCustomization.h
new file mode 100644
--- /dev/null
+++ b/Engine/Plugins/Editor/ModelingToolsEditorMode/Source/ModelingToolsEditorMode/Public/DetailsCustomizations/ModelingToolsBrushSizeCustomization.h
@@ -0,0 +1,27 @@
+// Copyright Epic Games, Inc. All Rights Reserved.
+
+#pragma once
+
+#include "IPropertyTypeCustomization.h"
+
+class IDetailLayoutBuilder;
+
+// defining this here so we can use in FModelingToolsBrushSizeCustomization, should move to some shared constants header/file
+class FSculptToolsUIConstants
+{
+public:
+	// sculpt panel packs in the widgets w/ short fixed-width labels
+	static constexpr int SculptShortLabelWidth = 40;
+
+};
+
+// Details customization for FBrushToolRadius struct.
+// This maybe should be removed in favor of just doing custom UI in FSculptBrushPropertiesDetails...
+class FModelingToolsBrushSizeCustomization : public IPropertyTypeCustomization
+{
+public:
+	static TSharedRef<IPropertyTypeCustomization> MakeInstance();
+
+	virtual void CustomizeHeader(TSharedRef<IPropertyHandle> StructPropertyHandle, class FDetailWidgetRow& HeaderRow, IPropertyTypeCustomizationUtils& StructCustomizationUtils) override;
+	virtual void CustomizeChildren(TSharedRef<IPropertyHandle> StructPropertyHandle, class IDetailChildrenBuilder& ChildBuilder, IPropertyTypeCustomizationUtils& StructCustomizationUtils) override;
+};
diff --git a/Engine/Plugins/Editor/ModelingToolsEditorMode/Source/ModelingToolsEditorMode/Public/ModelingToolsEditorModeModule.h b/Engine/Plugins/Editor/ModelingToolsEditorMode/Source/ModelingToolsEditorMode/Public/ModelingToolsEditorModeModule.h
--- a/Engine/Plugins/Editor/ModelingToolsEditorMode/Source/ModelingToolsEditorMode/Public/ModelingToolsEditorModeModule.h
+++ b/Engine/Plugins/Editor/ModelingToolsEditorMode/Source/ModelingToolsEditorMode/Public/ModelingToolsEditorModeModule.h
@@ -14,4 +14,9 @@ public:
 	virtual void ShutdownModule() override;
 
 	void OnPostEngineInit();
+
+private:
+	/** StaticClass is not safe on shutdown, so we cache the name, and use this to unregister on shut down */
+	TArray<FName> ClassesToUnregisterOnShutdown;
+	TArray<FName> PropertiesToUnregisterOnShutdown;
 };
diff --git a/Engine/Plugins/Editor/ModelingToolsEditorMode/Source/ModelingToolsEditorMode/Public/ModelingToolsEditorModeSettings.h b/Engine/Plugins/Editor/ModelingToolsEditorMode/Source/ModelingToolsEditorMode/Public/ModelingToolsEditorModeSettings.h
--- a/Engine/Plugins/Editor/ModelingToolsEditorMode/Source/ModelingToolsEditorMode/Public/ModelingToolsEditorModeSettings.h
+++ b/Engine/Plugins/Editor/ModelingToolsEditorMode/Source/ModelingToolsEditorMode/Public/ModelingToolsEditorModeSettings.h
@@ -175,6 +175,24 @@ struct FModelingModeCustomToolColor
 };
 
 
+/**
+ * Defines a Named list/set of content-browser Collection names
+ */
+USTRUCT()
+struct FModelingModeAssetCollectionSet
+{
+	GENERATED_BODY()
+
+	/** Name of the set of collections */
+	UPROPERTY(EditAnywhere, Category = "CollectionSet")
+	FString Name;
+
+	/** List of Collection names */
+	UPROPERTY(EditAnywhere, Category = "CollectionSet")
+	TArray<FCollectionReference> Collections;
+};
+
+
 UCLASS(config=Editor)
 class MODELINGTOOLSEDITORMODE_API UModelingToolsModeCustomizationSettings : public UDeveloperSettings
 {
@@ -214,4 +232,12 @@ public:
 	 */
 	UPROPERTY(config, EditAnywhere, Category = "Modeling Mode|UI Customization")
 	TArray<FModelingModeCustomToolColor> ToolColors;
+
+
+	/**
+	 * A Brush Alpha Set is a named list of Content Browser Collections, which will be shown as a separate tab in 
+	 * the Texture Asset Picker used in various Modeling Mode Tools when selecting a Brush Alpha (eg in Sculpting)
+	 */
+	UPROPERTY(config, EditAnywhere, Category = "Modeling Mode|Tool Assets")
+	TArray<FModelingModeAssetCollectionSet> BrushAlphaSets;
 };
diff --git a/Engine/Plugins/Editor/ModelingToolsEditorMode/Source/ModelingToolsEditorMode/Public/ModelingToolsEditorModeStyle.h b/Engine/Plugins/Editor/ModelingToolsEditorMode/Source/ModelingToolsEditorMode/Public/ModelingToolsEditorModeStyle.h
--- a/Engine/Plugins/Editor/ModelingToolsEditorMode/Source/ModelingToolsEditorMode/Public/ModelingToolsEditorModeStyle.h
+++ b/Engine/Plugins/Editor/ModelingToolsEditorMode/Source/ModelingToolsEditorMode/Public/ModelingToolsEditorModeStyle.h
@@ -16,6 +16,10 @@ public:
 	static TSharedPtr< class ISlateStyle > Get();
 
 	static FName GetStyleSetName();
+
+	// use to access icons defined by the style set by name, eg GetBrush("BrushFalloffIcons.Smooth")
+	static const FSlateBrush* GetBrush(FName PropertyName, const ANSICHAR* Specifier = NULL);
+
 private:
 	static FString InContent(const FString& RelativePath, const ANSICHAR* Extension);
 
