ModelingTools: Add new icons for Convert, CubeGrid, BakeRS, BakeTx, BakeAll, BakeVtx, Split, PivotAct

#rnx

diff --git a/Engine/Plugins/Editor/ModelingToolsEditorMode/Source/ModelingToolsEditorMode/Private/ModelingToolsEditorModeStyle.cpp b/Engine/Plugins/Editor/ModelingToolsEditorMode/Source/ModelingToolsEditorMode/Private/ModelingToolsEditorModeStyle.cpp
--- a/Engine/Plugins/Editor/ModelingToolsEditorMode/Source/ModelingToolsEditorMode/Private/ModelingToolsEditorModeStyle.cpp
+++ b/Engine/Plugins/Editor/ModelingToolsEditorMode/Source/ModelingToolsEditorMode/Private/ModelingToolsEditorModeStyle.cpp
@@ -171,12 +171,12 @@ void FModelingToolsEditorModeStyle::Initialize()
 		StyleSet->Set("ModelingToolsManagerCommands.BeginEditTangentsTool.Small", 		new IMAGE_PLUGIN_BRUSH("Icons/ModelingTangents_x40",			Icon20x20));
 		StyleSet->Set("ModelingToolsManagerCommands.BeginUVSeamEditTool", 			new IMAGE_PLUGIN_BRUSH("Icons/ModelingUVSeamEdit_x40",			Icon20x20));
 		StyleSet->Set("ModelingToolsManagerCommands.BeginUVSeamEditTool.Small", 		new IMAGE_PLUGIN_BRUSH("Icons/ModelingUVSeamEdit_x40",			Icon20x20));
-		StyleSet->Set("ModelingToolsManagerCommands.BeginBakeMeshAttributeMapsTool", 			new IMAGE_PLUGIN_BRUSH("Icons/ModelingBakeMaps_x40",			Icon20x20));
-		StyleSet->Set("ModelingToolsManagerCommands.BeginBakeMeshAttributeMapsTool.Small", 		new IMAGE_PLUGIN_BRUSH("Icons/ModelingBakeMaps_x40",			Icon20x20));
-		StyleSet->Set("ModelingToolsManagerCommands.BeginBakeMultiMeshAttributeMapsTool", 			new IMAGE_PLUGIN_BRUSH("Icons/ModelingBakeMaps_x40",			Icon20x20));
-		StyleSet->Set("ModelingToolsManagerCommands.BeginBakeMultiMeshAttributeMapsTool.Small", 		new IMAGE_PLUGIN_BRUSH("Icons/ModelingBakeMaps_x40",			Icon20x20));
-		StyleSet->Set("ModelingToolsManagerCommands.BeginBakeMeshAttributeVertexTool", new IMAGE_PLUGIN_BRUSH("Icons/ModelingBakeMaps_x40", Icon20x20));
-		StyleSet->Set("ModelingToolsManagerCommands.BeginBakeMeshAttributeVertexTool.Small", new IMAGE_PLUGIN_BRUSH("Icons/ModelingBakeMaps_x40", Icon20x20));
+		StyleSet->Set("ModelingToolsManagerCommands.BeginBakeMeshAttributeMapsTool", 			new IMAGE_BRUSH_SVG("Icons/BakeTexture",			Icon20x20));
+		StyleSet->Set("ModelingToolsManagerCommands.BeginBakeMeshAttributeMapsTool.Small", 		new IMAGE_BRUSH_SVG("Icons/BakeTexture",			Icon20x20));
+		StyleSet->Set("ModelingToolsManagerCommands.BeginBakeMultiMeshAttributeMapsTool", 			new IMAGE_BRUSH_SVG("Icons/BakeAll",			Icon20x20));
+		StyleSet->Set("ModelingToolsManagerCommands.BeginBakeMultiMeshAttributeMapsTool.Small", 		new IMAGE_BRUSH_SVG("Icons/BakeAll",			Icon20x20));
+		StyleSet->Set("ModelingToolsManagerCommands.BeginBakeMeshAttributeVertexTool", new IMAGE_BRUSH_SVG("Icons/BakeVertex", Icon20x20));
+		StyleSet->Set("ModelingToolsManagerCommands.BeginBakeMeshAttributeVertexTool.Small", new IMAGE_BRUSH_SVG("Icons/BakeVertex", Icon20x20));
 		//StyleSet->Set("ModelingToolsManagerCommands.BeginRemoveOccludedTrianglesTool", 				new IMAGE_PLUGIN_BRUSH("Icons/Jacket_40x",			Icon20x20));
 		//StyleSet->Set("ModelingToolsManagerCommands.BeginRemoveOccludedTrianglesTool.Small", 		new IMAGE_PLUGIN_BRUSH("Icons/Jacket_40x",			Icon20x20));
 		StyleSet->Set("ModelingToolsManagerCommands.BeginHoleFillTool", 				new IMAGE_PLUGIN_BRUSH("Icons/ModelingHoleFill_x40",			Icon20x20));
@@ -217,8 +217,8 @@ void FModelingToolsEditorModeStyle::Initialize()
 		StyleSet->Set("ModelingToolsManagerCommands.BeginTransferMeshTool.Small",            new FSlateImageBrush(StyleSet->RootToCoreContentDir(TEXT("../Editor/Slate/Icons/GeneralTools/Next_40x.png")), Icon20x20));
 		StyleSet->Set("ModelingToolsManagerCommands.BeginGlobalUVGenerateTool",              new IMAGE_PLUGIN_BRUSH("Icons/AutoUnwrap_40x",       Icon20x20));      
 		StyleSet->Set("ModelingToolsManagerCommands.BeginGlobalUVGenerateTool.Small",        new IMAGE_PLUGIN_BRUSH("Icons/AutoUnwrap_40x",       Icon20x20));      
-		StyleSet->Set("ModelingToolsManagerCommands.BeginBakeTransformTool",                 new IMAGE_PLUGIN_BRUSH("Icons/BakeXForm_40x",        Icon20x20));      
-		StyleSet->Set("ModelingToolsManagerCommands.BeginBakeTransformTool.Small",           new IMAGE_PLUGIN_BRUSH("Icons/BakeXForm_40x",        Icon20x20));      
+		StyleSet->Set("ModelingToolsManagerCommands.BeginBakeTransformTool",                 new IMAGE_BRUSH_SVG("Icons/GeometryBakeXForm",        Icon20x20));      
+		StyleSet->Set("ModelingToolsManagerCommands.BeginBakeTransformTool.Small",           new IMAGE_BRUSH_SVG("Icons/GeometryBakeXForm",        Icon20x20));      
 		StyleSet->Set("ModelingToolsManagerCommands.BeginCombineMeshesTool",                 new IMAGE_BRUSH_SVG("Icons/GeometryCombine",          Icon20x20));
 		StyleSet->Set("ModelingToolsManagerCommands.BeginCombineMeshesTool.Small",           new IMAGE_BRUSH_SVG("Icons/GeometryCombine",          Icon20x20));
 		StyleSet->Set("ModelingToolsManagerCommands.BeginDuplicateMeshesTool",               new IMAGE_PLUGIN_BRUSH("Icons/Duplicate_40x",        Icon20x20));   
@@ -226,7 +226,9 @@ void FModelingToolsEditorModeStyle::Initialize()
 		StyleSet->Set("ModelingToolsManagerCommands.BeginEditMeshMaterialsTool",             new IMAGE_PLUGIN_BRUSH("Icons/EditMats_40x",         Icon20x20));     
 		StyleSet->Set("ModelingToolsManagerCommands.BeginEditMeshMaterialsTool.Small",       new IMAGE_PLUGIN_BRUSH("Icons/EditMats_40x",         Icon20x20));     
 		StyleSet->Set("ModelingToolsManagerCommands.BeginEditPivotTool",                     new IMAGE_PLUGIN_BRUSH("Icons/EditPivot_40x",        Icon20x20));      
-		StyleSet->Set("ModelingToolsManagerCommands.BeginEditPivotTool.Small",               new IMAGE_PLUGIN_BRUSH("Icons/EditPivot_40x",        Icon20x20));      
+		StyleSet->Set("ModelingToolsManagerCommands.BeginEditPivotTool.Small",               new IMAGE_PLUGIN_BRUSH("Icons/EditPivot_40x",        Icon20x20));
+		StyleSet->Set("ModelingToolsManagerCommands.BeginAddPivotActorTool",                 new IMAGE_BRUSH_SVG("Icons/Pivot",                   Icon20x20));
+		StyleSet->Set("ModelingToolsManagerCommands.BeginAddPivotActorTool.Small",           new IMAGE_BRUSH_SVG("Icons/Pivot",                   Icon20x20));
 		StyleSet->Set("ModelingToolsManagerCommands.BeginGroupUVGenerateTool",               new IMAGE_PLUGIN_BRUSH("Icons/GroupUnwrap_40x",      Icon20x20));       
 		StyleSet->Set("ModelingToolsManagerCommands.BeginGroupUVGenerateTool.Small",         new IMAGE_PLUGIN_BRUSH("Icons/GroupUnwrap_40x",      Icon20x20));       
 		StyleSet->Set("ModelingToolsManagerCommands.BeginRemoveOccludedTrianglesTool",       new IMAGE_PLUGIN_BRUSH("Icons/Jacketing_40x",        Icon20x20));     
@@ -243,6 +245,8 @@ void FModelingToolsEditorModeStyle::Initialize()
 		StyleSet->Set("ModelingToolsManagerCommands.BeginDrawAndRevolveTool.Small",          new IMAGE_PLUGIN_BRUSH("Icons/ModelingDrawAndRevolve_x20", Icon20x20));
 		StyleSet->Set("ModelingToolsManagerCommands.BeginRevolveBoundaryTool",               new IMAGE_PLUGIN_BRUSH("Icons/ModelingRevolveBoundary_x40", Icon20x20));
 		StyleSet->Set("ModelingToolsManagerCommands.BeginRevolveBoundaryTool.Small",         new IMAGE_PLUGIN_BRUSH("Icons/ModelingRevolveBoundary_x20", Icon20x20));
+		StyleSet->Set("ModelingToolsManagerCommands.BeginCubeGridTool",                      new IMAGE_BRUSH_SVG("Icons/CubeGrid",                Icon20x20));
+		StyleSet->Set("ModelingToolsManagerCommands.BeginCubeGridTool.Small",                new IMAGE_BRUSH_SVG("Icons/CubeGrid",                Icon20x20));
 		StyleSet->Set("ModelingToolsManagerCommands.BeginMeshBooleanTool",                   new IMAGE_PLUGIN_BRUSH("Icons/ModelingMeshBoolean_x40", Icon20x20));
 		StyleSet->Set("ModelingToolsManagerCommands.BeginMeshBooleanTool.Small",             new IMAGE_PLUGIN_BRUSH("Icons/ModelingMeshBoolean_x20", Icon20x20));
 		StyleSet->Set("ModelingToolsManagerCommands.BeginMeshTrimTool",				 new IMAGE_BRUSH_SVG("Icons/ModelingTrim", Icon20x20));
@@ -269,6 +273,10 @@ void FModelingToolsEditorModeStyle::Initialize()
 		StyleSet->Set("ModelingToolsManagerCommands.BeginMeshGroupPaintTool.Small", new IMAGE_BRUSH_SVG("Icons/GroupPaint", Icon20x20));
 		StyleSet->Set("ModelingToolsManagerCommands.BeginLatticeDeformerTool", new IMAGE_BRUSH_SVG("Icons/LatticeDeformation", Icon20x20));
 		StyleSet->Set("ModelingToolsManagerCommands.BeginLatticeDeformerTool.Small", new IMAGE_BRUSH_SVG("Icons/LatticeDeformation", Icon20x20));
+		StyleSet->Set("ModelingToolsManagerCommands.BeginConvertMeshesTool", new IMAGE_BRUSH_SVG("Icons/Convert_20", Icon20x20));
+		StyleSet->Set("ModelingToolsManagerCommands.BeginConvertMeshesTool.Small", new IMAGE_BRUSH_SVG("Icons/Convert_20", Icon20x20));
+		StyleSet->Set("ModelingToolsManagerCommands.BeginSplitMeshesTool", new IMAGE_BRUSH_SVG("Icons/GeometrySplit", Icon20x20));
+		StyleSet->Set("ModelingToolsManagerCommands.BeginSplitMeshesTool.Small", new IMAGE_BRUSH_SVG("Icons/GeometrySplit", Icon20x20));
 
 		
 
