ModelingTools: Update Bake tools descriptions

#rnx

diff --git a/Engine/Plugins/Editor/ModelingToolsEditorMode/Source/ModelingToolsEditorMode/Private/ModelingToolsManagerActions.cpp b/Engine/Plugins/Editor/ModelingToolsEditorMode/Source/ModelingToolsEditorMode/Private/ModelingToolsManagerActions.cpp
--- a/Engine/Plugins/Editor/ModelingToolsEditorMode/Source/ModelingToolsEditorMode/Private/ModelingToolsManagerActions.cpp
+++ b/Engine/Plugins/Editor/ModelingToolsEditorMode/Source/ModelingToolsEditorMode/Private/ModelingToolsManagerActions.cpp
@@ -69,7 +69,30 @@ void FModelingToolsManagerCommands::RegisterCommands()
 	REGISTER_MODELING_TOOL_COMMAND(BeginMeshBooleanTool, "MshBool", "Apply Boolean operations to mesh pairs");
 	REGISTER_MODELING_TOOL_COMMAND(BeginCutMeshWithMeshTool, "MshCut", "Split one mesh into parts using a second mesh");
 	REGISTER_MODELING_TOOL_COMMAND(BeginSubdividePolyTool, "SubDiv", "Subdivide Mesh via PolyGroups or Triangles");
-	
+
+	// TriModel
+
+	// Deform
+
+	// Transform
+
+	// MeshOps
+
+	// VoxOps
+
+	// Attributes
+
+	// UVs
+
+	// Baking
+	REGISTER_MODELING_TOOL_COMMAND(BeginBakeMeshAttributeMapsTool, "BakeTx", "Bake textures for single meshes");
+	REGISTER_MODELING_TOOL_COMMAND(BeginBakeMultiMeshAttributeMapsTool, "BakeAll", "Bake textures for single meshes from multiple source meshes");
+	REGISTER_MODELING_TOOL_COMMAND(BeginBakeMeshAttributeVertexTool, "BakeVtx", "Bake vertex colors for single meshes");
+
+	// Volumes
+
+	// LODs
+
 	REGISTER_MODELING_TOOL_COMMAND(BeginAddPatchTool, "Patch", "Add Patch");
 	REGISTER_MODELING_TOOL_COMMAND(BeginShapeSprayTool, "Spray", "Shape Spray");
 
@@ -128,9 +151,6 @@ void FModelingToolsManagerCommands::RegisterCommands()
 	REGISTER_MODELING_TOOL_COMMAND(BeginTransformUVIslandsTool, "XForm", "Transform UV Islands in UV Space");
 	REGISTER_MODELING_TOOL_COMMAND(BeginMeshAttributePaintTool, "MapPnt", "Paint Attribute Maps");
 	REGISTER_MODELING_TOOL_COMMAND(BeginAttributeEditorTool, "AttrEd", "Edit/Configure Mesh Attributes");
-	REGISTER_MODELING_TOOL_COMMAND(BeginBakeMeshAttributeMapsTool, "BakeTx", "Bake Image Maps for Target Mesh (optionally from second Source Mesh)");
-	REGISTER_MODELING_TOOL_COMMAND(BeginBakeMultiMeshAttributeMapsTool, "BakeAll", "Bake Image Maps for Target Mesh from multiple Source/Detail meshes");
-	REGISTER_MODELING_TOOL_COMMAND(BeginBakeMeshAttributeVertexTool, "BakeVtx", "Bake Vertex Colors for Target Mesh (optionally from a second Source Mesh)");
 	REGISTER_MODELING_TOOL_COMMAND(BeginUVSeamEditTool, "SeamEd", "Add UV Seams to Mesh");
 
 	// why are these ::Button ?
