ModelingTools: Disable skeletal mesh tool target support in Modeling Mode.

#rnx

diff --git a/Engine/Plugins/Editor/ModelingToolsEditorMode/Source/ModelingToolsEditorMode/Private/ModelingToolsEditorMode.cpp b/Engine/Plugins/Editor/ModelingToolsEditorMode/Source/ModelingToolsEditorMode/Private/ModelingToolsEditorMode.cpp
--- a/Engine/Plugins/Editor/ModelingToolsEditorMode/Source/ModelingToolsEditorMode/Private/ModelingToolsEditorMode.cpp
+++ b/Engine/Plugins/Editor/ModelingToolsEditorMode/Source/ModelingToolsEditorMode/Private/ModelingToolsEditorMode.cpp
@@ -319,7 +319,11 @@ void UModelingToolsEditorMode::Enter()
 	GetInteractiveToolsContext()->TargetManager->AddTargetFactory(NewObject<UVolumeDynamicMeshToolTargetFactory>(GetToolManager()));
 	GetInteractiveToolsContext()->TargetManager->AddTargetFactory(NewObject<UVolumeMeshDescriptionToolTargetFactory>(GetToolManager()));
 	GetInteractiveToolsContext()->TargetManager->AddTargetFactory(NewObject<UDynamicMeshComponentToolTargetFactory>(GetToolManager()));
-	GetInteractiveToolsContext()->TargetManager->AddTargetFactory(NewObject<USkeletalMeshComponentToolTargetFactory>(GetToolManager()));
+
+	// TODO: Re-enable skeletal mesh target support when the tool target system has a better way to filter/exclude targets from certain tools.
+	// Skeletal mesh targets are disabled since many modeling tools would break skin weights. In the interim, tools
+	// that require skeletal mesh support should manually handle it at the tool level.
+	//GetInteractiveToolsContext()->TargetManager->AddTargetFactory(NewObject<USkeletalMeshComponentToolTargetFactory>(GetToolManager()));
 
 	// register stylus event handler
 	StylusStateTracker = MakeUnique<FStylusStateTracker>();
