MeshModelingToolset: Cleanup UV Layout tool

diff --git a/Engine/Plugins/Editor/ModelingToolsEditorMode/Source/ModelingToolsEditorMode/Private/ModelingToolsManagerActions.cpp b/Engine/Plugins/Editor/ModelingToolsEditorMode/Source/ModelingToolsEditorMode/Private/ModelingToolsManagerActions.cpp
--- a/Engine/Plugins/Editor/ModelingToolsEditorMode/Source/ModelingToolsEditorMode/Private/ModelingToolsManagerActions.cpp
+++ b/Engine/Plugins/Editor/ModelingToolsEditorMode/Source/ModelingToolsEditorMode/Private/ModelingToolsManagerActions.cpp
@@ -83,6 +83,12 @@ void FModelingToolsManagerCommands::RegisterCommands()
 	// Attributes
 
 	// UVs
+	REGISTER_MODELING_TOOL_COMMAND(BeginGlobalUVGenerateTool, "AutoUV", "Auto-Unwrap and pack UVs for Mesh");
+	REGISTER_MODELING_TOOL_COMMAND(BeginGroupUVGenerateTool, "Unwrap", "Recalculate UV unwrapping for Mesh regions");
+	REGISTER_MODELING_TOOL_COMMAND(BeginUVProjectionTool, "Project", "Set UVs from Projection");
+	REGISTER_MODELING_TOOL_COMMAND(BeginUVSeamEditTool, "SeamEd", "Add UV Seams to Mesh");
+	REGISTER_MODELING_TOOL_COMMAND(BeginTransformUVIslandsTool, "XForm", "Transform UV islands in UV space");
+	REGISTER_MODELING_TOOL_COMMAND(BeginUVLayoutTool, "Layout", "Transform and pack existing UVs");
 
 	// Baking
 	REGISTER_MODELING_TOOL_COMMAND(BeginBakeMeshAttributeMapsTool, "BakeTx", "Bake textures for single meshes");
@@ -118,8 +124,6 @@ void FModelingToolsManagerCommands::RegisterCommands()
 	REGISTER_MODELING_TOOL_COMMAND(BeginEditNormalsTool, "Nrmls", "Modify/Update Mesh Normals");
 	REGISTER_MODELING_TOOL_COMMAND(BeginEditTangentsTool, "Tngnts", "Update Mesh Tangents");
 	REGISTER_MODELING_TOOL_COMMAND(BeginRemoveOccludedTrianglesTool, "Jacket", "Remove Hidden Triangles from selected Meshes");
-	REGISTER_MODELING_TOOL_COMMAND(BeginUVProjectionTool, "Project", "Set UVs from Projection");
-	REGISTER_MODELING_TOOL_COMMAND(BeginUVLayoutTool, "Layout", "Re-Pack Existing UVs");
 	REGISTER_MODELING_TOOL_COMMAND(BeginVoxelMergeTool, "VoxMrg", "Merge Selected Meshes (Voxel-Based)");
 	REGISTER_MODELING_TOOL_COMMAND(BeginVoxelBooleanTool, "VoxBool", "Boolean Selected Meshes (Voxel-Based)");
 	REGISTER_MODELING_TOOL_COMMAND(BeginVoxelSolidifyTool, "VoxWrap", "Wrap Selected Meshes (Voxel-Based)");
@@ -135,8 +139,7 @@ void FModelingToolsManagerCommands::RegisterCommands()
 	REGISTER_MODELING_TOOL_COMMAND(BeginHoleFillTool, "HFill", "Fill Holes in Mesh");
 	REGISTER_MODELING_TOOL_COMMAND(BeginLatticeDeformerTool, "Lattice", "Deform Mesh with 3D Lattice/Grid");
 	REGISTER_MODELING_TOOL_COMMAND(BeginPolygonCutTool, "PolyCut", "Cut Mesh with Extruded Polygon");
-	REGISTER_MODELING_TOOL_COMMAND(BeginGlobalUVGenerateTool, "AutoUV", "Auto-Unwrap and Pack UVs for Mesh");
-	REGISTER_MODELING_TOOL_COMMAND(BeginGroupUVGenerateTool, "Unwrap", "Recalculate UV unwrapping for Mesh regions");
+	
 	REGISTER_MODELING_TOOL_COMMAND(BeginMeshSelectionTool, "TriSel", "Select and Edit Mesh Triangles");
 
 	REGISTER_MODELING_TOOL_COMMAND(BeginPhysicsInspectorTool, "PInspct", "Inspect Physics Geometry for selected Meshes");
@@ -148,10 +151,8 @@ void FModelingToolsManagerCommands::RegisterCommands()
 	REGISTER_MODELING_TOOL_COMMAND(BeginWeldEdgesTool, "Weld", "Weld Overlapping Mesh Edges");
 	REGISTER_MODELING_TOOL_COMMAND(BeginPolyGroupsTool, "GenGrps", "Generate new PolyGroups");
 	REGISTER_MODELING_TOOL_COMMAND(BeginEditMeshMaterialsTool, "MatEd", "Assign Materials to Selected Triangles");
-	REGISTER_MODELING_TOOL_COMMAND(BeginTransformUVIslandsTool, "XForm", "Transform UV Islands in UV Space");
 	REGISTER_MODELING_TOOL_COMMAND(BeginMeshAttributePaintTool, "MapPnt", "Paint Attribute Maps");
 	REGISTER_MODELING_TOOL_COMMAND(BeginAttributeEditorTool, "AttrEd", "Edit/Configure Mesh Attributes");
-	REGISTER_MODELING_TOOL_COMMAND(BeginUVSeamEditTool, "SeamEd", "Add UV Seams to Mesh");
 
 	// why are these ::Button ?
 	UI_COMMAND(BeginSkinWeightsPaintTool, "SkinWts", "Start the Paint Skin Weights Tool", EUserInterfaceActionType::Button, FInputChord());
diff --git a/Engine/Plugins/Experimental/MeshModelingToolsetExp/Source/MeshModelingToolsEditorOnly/Private/ParameterizeMeshTool.cpp b/Engine/Plugins/Experimental/MeshModelingToolsetExp/Source/MeshModelingToolsEditorOnly/Private/ParameterizeMeshTool.cpp
--- a/Engine/Plugins/Experimental/MeshModelingToolsetExp/Source/MeshModelingToolsEditorOnly/Private/ParameterizeMeshTool.cpp
+++ b/Engine/Plugins/Experimental/MeshModelingToolsetExp/Source/MeshModelingToolsEditorOnly/Private/ParameterizeMeshTool.cpp
@@ -6,7 +6,6 @@
 #include "ToolSetupUtil.h"
 #include "ModelingToolTargetUtil.h"
 #include "DynamicMesh/DynamicMesh3.h"
-#include "FaceGroupUtil.h"
 #include "ParameterizationOps/ParameterizeMeshOp.h"
 #include "Properties/ParameterizeMeshProperties.h"
 
@@ -61,7 +60,8 @@ void UParameterizeMeshTool::Setup()
 	UVChannelProperties->ValidateSelection(true);
 	UVChannelProperties->WatchProperty(UVChannelProperties->UVChannel, [this](const FString& NewValue) 
 	{
-		MaterialSettings->UVChannel = UVChannelProperties->GetSelectedChannelIndex(true);
+		MaterialSettings->UpdateUVChannels(UVChannelProperties->UVChannelNamesList.IndexOfByKey(UVChannelProperties->UVChannel),
+		                                   UVChannelProperties->UVChannelNamesList);
 	});
 	AddToolPropertySource(UVChannelProperties);
 
@@ -107,7 +107,7 @@ void UParameterizeMeshTool::Setup()
 		UVLayoutView->CreateInWorld(TargetWorld);
 		UVLayoutView->SetSourceMaterials(MaterialSet);
 		UVLayoutView->SetSourceWorldPosition(InputTransform, UE::ToolTarget::GetTargetActor(Target)->GetComponentsBoundingBox());
-		UVLayoutView->Settings->bVisible = false;
+		UVLayoutView->Settings->bEnabled = false;
 		UVLayoutView->Settings->bShowWireframe = false;
 		UVLayoutView->Settings->RestoreProperties(this, TEXT("ParameterizeMeshTool"));
 		AddToolPropertySource(UVLayoutView->Settings);
diff --git a/Engine/Plugins/Experimental/MeshModelingToolsetExp/Source/MeshModelingToolsExp/Private/Properties/MeshUVChannelProperties.cpp b/Engine/Plugins/Experimental/MeshModelingToolsetExp/Source/MeshModelingToolsExp/Private/Properties/MeshUVChannelProperties.cpp
--- a/Engine/Plugins/Experimental/MeshModelingToolsetExp/Source/MeshModelingToolsExp/Private/Properties/MeshUVChannelProperties.cpp
+++ b/Engine/Plugins/Experimental/MeshModelingToolsetExp/Source/MeshModelingToolsExp/Private/Properties/MeshUVChannelProperties.cpp
@@ -7,8 +7,6 @@
 
 #include "MeshDescription.h"
 #include "StaticMeshAttributes.h"
-#include "StaticMeshOperations.h"
-
 
 
 void UMeshUVChannelProperties::Initialize(int32 NumUVChannels, bool bInitializeSelection)
@@ -16,7 +14,7 @@ void UMeshUVChannelProperties::Initialize(int32 NumUVChannels, bool bInitializeS
 	UVChannelNamesList.Reset();
 	for (int32 k = 0; k < NumUVChannels; ++k)
 	{
-		UVChannelNamesList.Add(FString::Printf(TEXT("UV%d"), k));
+		UVChannelNamesList.Add(FString::Printf(TEXT("UV %d"), k));
 	}
 	if (bInitializeSelection)
 	{
@@ -25,7 +23,7 @@ void UMeshUVChannelProperties::Initialize(int32 NumUVChannels, bool bInitializeS
 }
 
 
-const TArray<FString>& UMeshUVChannelProperties::GetUVChannelNamesFunc()
+const TArray<FString>& UMeshUVChannelProperties::GetUVChannelNamesFunc() const
 {
 	return UVChannelNamesList;
 }
diff --git a/Engine/Plugins/Experimental/MeshModelingToolsetExp/Source/MeshModelingToolsExp/Private/RecomputeUVsTool.cpp b/Engine/Plugins/Experimental/MeshModelingToolsetExp/Source/MeshModelingToolsExp/Private/RecomputeUVsTool.cpp
--- a/Engine/Plugins/Experimental/MeshModelingToolsetExp/Source/MeshModelingToolsExp/Private/RecomputeUVsTool.cpp
+++ b/Engine/Plugins/Experimental/MeshModelingToolsetExp/Source/MeshModelingToolsExp/Private/RecomputeUVsTool.cpp
@@ -5,7 +5,6 @@
 #include "ToolBuilderUtil.h"
 #include "DynamicMesh/DynamicMesh3.h"
 #include "Polygroups/PolygroupUtil.h"
-#include "FaceGroupUtil.h"
 #include "ToolSetupUtil.h"
 #include "ModelingToolTargetUtil.h"
 #include "ParameterizationOps/RecomputeUVsOp.h"
@@ -52,7 +51,8 @@ void URecomputeUVsTool::Setup()
 	UVChannelProperties->ValidateSelection(true);
 	UVChannelProperties->WatchProperty(UVChannelProperties->UVChannel, [this](const FString& NewValue)
 		{
-			MaterialSettings->UVChannel = UVChannelProperties->GetSelectedChannelIndex(true);
+			MaterialSettings->UpdateUVChannels(UVChannelProperties->UVChannelNamesList.IndexOfByKey(UVChannelProperties->UVChannel),
+			                                   UVChannelProperties->UVChannelNamesList);
 		});
 	AddToolPropertySource(UVChannelProperties);
 
@@ -100,7 +100,7 @@ void URecomputeUVsTool::Setup()
 		UVLayoutView->CreateInWorld(TargetWorld);
 		UVLayoutView->SetSourceMaterials(MaterialSet);
 		UVLayoutView->SetSourceWorldPosition(TargetTransform, UE::ToolTarget::GetTargetActor(Target)->GetComponentsBoundingBox());
-		UVLayoutView->Settings->bVisible = false;
+		UVLayoutView->Settings->bEnabled = false;
 		UVLayoutView->Settings->bShowWireframe = false;
 		UVLayoutView->Settings->RestoreProperties(this, TEXT("RecomputeUVsTool"));
 		AddToolPropertySource(UVLayoutView->Settings);
diff --git a/Engine/Plugins/Experimental/MeshModelingToolsetExp/Source/MeshModelingToolsExp/Private/UVLayoutTool.cpp b/Engine/Plugins/Experimental/MeshModelingToolsetExp/Source/MeshModelingToolsExp/Private/UVLayoutTool.cpp
--- a/Engine/Plugins/Experimental/MeshModelingToolsetExp/Source/MeshModelingToolsExp/Private/UVLayoutTool.cpp
+++ b/Engine/Plugins/Experimental/MeshModelingToolsetExp/Source/MeshModelingToolsExp/Private/UVLayoutTool.cpp
@@ -7,12 +7,9 @@
 #include "ToolSetupUtil.h"
 
 #include "DynamicMesh/DynamicMesh3.h"
-#include "BaseBehaviors/MultiClickSequenceInputBehavior.h"
-#include "Selection/SelectClickedAction.h"
 
 #include "MeshDescriptionToDynamicMesh.h"
 #include "DynamicMeshToMeshDescription.h"
-#include "Materials/MaterialInstanceDynamic.h"
 
 #include "ParameterizationOps/UVLayoutOp.h"
 #include "Properties/UVLayoutProperties.h"
@@ -94,7 +91,8 @@ void UUVLayoutTool::Setup()
 		AddToolPropertySource(UVChannelProperties);
 		UVChannelProperties->WatchProperty(UVChannelProperties->UVChannel, [this](const FString& NewValue)
 		{
-			MaterialSettings->UVChannel = UVChannelProperties->GetSelectedChannelIndex(true);
+			MaterialSettings->UpdateUVChannels(UVChannelProperties->UVChannelNamesList.IndexOfByKey(UVChannelProperties->UVChannel),
+			                                   UVChannelProperties->UVChannelNamesList);
 			UpdateVisualization();
 		});
 	}
@@ -135,8 +133,8 @@ void UUVLayoutTool::Setup()
 
 void UUVLayoutTool::UpdateNumPreviews()
 {
-	int32 CurrentNumPreview = Previews.Num();
-	int32 TargetNumPreview = Targets.Num();
+	const int32 CurrentNumPreview = Previews.Num();
+	const int32 TargetNumPreview = Targets.Num();
 	if (TargetNumPreview < CurrentNumPreview)
 	{
 		for (int32 PreviewIdx = CurrentNumPreview - 1; PreviewIdx >= TargetNumPreview; PreviewIdx--)
@@ -303,7 +301,7 @@ void UUVLayoutTool::UpdateVisualization()
 
 bool UUVLayoutTool::CanAccept() const
 {
-	for (UMeshOpPreviewWithBackgroundCompute* Preview : Previews)
+	for (const UMeshOpPreviewWithBackgroundCompute* Preview : Previews)
 	{
 		if (!Preview->HaveValidResult())
 		{
@@ -325,11 +323,11 @@ void UUVLayoutTool::GenerateAsset(const TArray<FDynamicMeshOpResult>& Results)
 		check(Results[ComponentIdx].Mesh.Get() != nullptr);
 		TargetMeshCommitterInterface(ComponentIdx)->CommitMeshDescription([&Results, &ComponentIdx, this](const IMeshDescriptionCommitter::FCommitterParams& CommitParams)
 		{
-			FDynamicMesh3* DynamicMesh = Results[ComponentIdx].Mesh.Get();
+			const FDynamicMesh3* DynamicMesh = Results[ComponentIdx].Mesh.Get();
 			FMeshDescription* MeshDescription = CommitParams.MeshDescriptionOut;
 			
-			bool bVerticesOnly = false;
-			bool bAttributesOnly = true;
+			constexpr bool bVerticesOnly = false;
+			constexpr bool bAttributesOnly = true;
 			if (FDynamicMeshToMeshDescription::HaveMatchingElementCounts(DynamicMesh, MeshDescription, bVerticesOnly, bAttributesOnly))
 			{
 				FDynamicMeshToMeshDescription Converter;
diff --git a/Engine/Plugins/Experimental/MeshModelingToolsetExp/Source/MeshModelingToolsExp/Private/UVProjectionTool.cpp b/Engine/Plugins/Experimental/MeshModelingToolsetExp/Source/MeshModelingToolsExp/Private/UVProjectionTool.cpp
--- a/Engine/Plugins/Experimental/MeshModelingToolsetExp/Source/MeshModelingToolsExp/Private/UVProjectionTool.cpp
+++ b/Engine/Plugins/Experimental/MeshModelingToolsetExp/Source/MeshModelingToolsExp/Private/UVProjectionTool.cpp
@@ -82,7 +82,8 @@ void UUVProjectionTool::Setup()
 	UVChannelProperties->ValidateSelection(true);
 	UVChannelProperties->WatchProperty(UVChannelProperties->UVChannel, [this](const FString& NewValue)
 	{
-		MaterialSettings->UVChannel = UVChannelProperties->GetSelectedChannelIndex(true);
+		MaterialSettings->UpdateUVChannels(UVChannelProperties->UVChannelNamesList.IndexOfByKey(UVChannelProperties->UVChannel),
+		                                   UVChannelProperties->UVChannelNamesList);
 		Preview->InvalidateResult();
 		OnMaterialSettingsChanged();
 	});
diff --git a/Engine/Plugins/Experimental/MeshModelingToolsetExp/Source/MeshModelingToolsExp/Public/Properties/MeshUVChannelProperties.h b/Engine/Plugins/Experimental/MeshModelingToolsetExp/Source/MeshModelingToolsExp/Public/Properties/MeshUVChannelProperties.h
--- a/Engine/Plugins/Experimental/MeshModelingToolsetExp/Source/MeshModelingToolsExp/Public/Properties/MeshUVChannelProperties.h
+++ b/Engine/Plugins/Experimental/MeshModelingToolsetExp/Source/MeshModelingToolsExp/Public/Properties/MeshUVChannelProperties.h
@@ -17,17 +17,16 @@ class MESHMODELINGTOOLSEXP_API UMeshUVChannelProperties : public UInteractiveToo
 	GENERATED_BODY()
 public:
 
-	UPROPERTY(EditAnywhere, Category = UVChannel, meta = (DisplayName = "UV Channel", GetOptions = GetUVChannelNamesFunc))
+	/** Select UV channel in the mesh */
+	UPROPERTY(EditAnywhere, Category = "UV Channel", meta = (DisplayName = "UV Channel", GetOptions = GetUVChannelNamesFunc, NoResetToDefault))
 	FString UVChannel;
 
 	UFUNCTION()
-	const TArray<FString>& GetUVChannelNamesFunc();
+	const TArray<FString>& GetUVChannelNamesFunc() const;
 
 	UPROPERTY(meta = (TransientToolProperty))
 	TArray<FString> UVChannelNamesList;
 
-
-public:
 	void Initialize(int32 NumUVChannels, bool bInitializeSelection = true);
 	void Initialize(const FMeshDescription* MeshDescription, bool bInitializeSelection = true);
 	void Initialize(const FDynamicMesh3* Mesh, bool bInitializeSelection = true);
diff --git a/Engine/Plugins/Experimental/MeshModelingToolsetExp/Source/MeshModelingToolsExp/Public/UVLayoutTool.h b/Engine/Plugins/Experimental/MeshModelingToolsetExp/Source/MeshModelingToolsExp/Public/UVLayoutTool.h
--- a/Engine/Plugins/Experimental/MeshModelingToolsetExp/Source/MeshModelingToolsExp/Public/UVLayoutTool.h
+++ b/Engine/Plugins/Experimental/MeshModelingToolsetExp/Source/MeshModelingToolsExp/Public/UVLayoutTool.h
@@ -10,7 +10,6 @@
 #include "MeshOpPreviewHelpers.h"
 #include "ToolDataVisualizer.h"
 #include "DynamicMesh/DynamicMesh3.h"
-#include "BaseTools/SingleClickTool.h"
 #include "Properties/MeshMaterialProperties.h"
 #include "Properties/MeshUVChannelProperties.h"
 #include "Drawing/UVLayoutPreview.h"
@@ -18,7 +17,7 @@
 #include "UVLayoutTool.generated.h"
 
 
-// predeclarations
+// Forward declarations
 struct FMeshDescription;
 class UDynamicMeshComponent;
 class UUVLayoutProperties;
@@ -68,11 +67,8 @@ public:
 
 	virtual void OnPropertyModified(UObject* PropertySet, FProperty* Property) override;
 
-
-public:
 	int32 GetSelectedUVChannel() const;
 
-
 protected:
 
 	UPROPERTY()
@@ -90,7 +86,6 @@ protected:
 	UPROPERTY()
 	TArray<TObjectPtr<UUVLayoutOperatorFactory>> Factories;
 
-protected:
 	TArray<TSharedPtr<UE::Geometry::FDynamicMesh3, ESPMode::ThreadSafe>> OriginalDynamicMeshes;
 
 	UWorld* TargetWorld = nullptr;
@@ -105,8 +100,6 @@ protected:
 
 	void GenerateAsset(const TArray<FDynamicMeshOpResult>& Results);
 
-
-protected:
 	UPROPERTY()
 	TObjectPtr<UUVLayoutPreview> UVLayoutView = nullptr;
 };
diff --git a/Engine/Plugins/Experimental/UVEditor/Source/UVEditorToolsEditorOnly/Private/UVEditorParameterizeMeshTool.cpp b/Engine/Plugins/Experimental/UVEditor/Source/UVEditorToolsEditorOnly/Private/UVEditorParameterizeMeshTool.cpp
--- a/Engine/Plugins/Experimental/UVEditor/Source/UVEditorToolsEditorOnly/Private/UVEditorParameterizeMeshTool.cpp
+++ b/Engine/Plugins/Experimental/UVEditor/Source/UVEditorToolsEditorOnly/Private/UVEditorParameterizeMeshTool.cpp
@@ -78,7 +78,7 @@ void UUVEditorParameterizeMeshTool::Setup()
 	MaterialSettings = NewObject<UExistingMeshMaterialProperties>(this);
 	MaterialSettings->MaterialMode = ESetMeshMaterialMode::Checkerboard;
 	MaterialSettings->RestoreProperties(this, TEXT("ModelingUVTools"));
-	MaterialSettings->UVChannel = 0; // We only ever have one UV layer for the applied preview.
+	MaterialSettings->UpdateUVChannels(0, {TEXT("UV 0")}); // We only ever have one UV layer for the applied preview.
 	AddToolPropertySource(MaterialSettings);
 	// force update
 	MaterialSettings->UpdateMaterials();
diff --git a/Engine/Plugins/Runtime/MeshModelingToolset/Source/MeshModelingTools/Private/Properties/MeshMaterialProperties.cpp b/Engine/Plugins/Runtime/MeshModelingToolset/Source/MeshModelingTools/Private/Properties/MeshMaterialProperties.cpp
--- a/Engine/Plugins/Runtime/MeshModelingToolset/Source/MeshModelingTools/Private/Properties/MeshMaterialProperties.cpp
+++ b/Engine/Plugins/Runtime/MeshModelingToolset/Source/MeshModelingTools/Private/Properties/MeshMaterialProperties.cpp
@@ -13,6 +13,11 @@ UNewMeshMaterialProperties::UNewMeshMaterialProperties()
 	Material = CreateDefaultSubobject<UMaterialInterface>(TEXT("MATERIAL"));
 }
 
+const TArray<FString>& UExistingMeshMaterialProperties::GetUVChannelNamesFunc() const
+{
+	return UVChannelNamesList;
+}
+
 void UExistingMeshMaterialProperties::RestoreProperties(UInteractiveTool* RestoreToTool, const FString& CacheIdentifier)
 {
 	Super::RestoreProperties(RestoreToTool, CacheIdentifier);
@@ -28,7 +33,7 @@ void UExistingMeshMaterialProperties::Setup()
 		if (CheckerMaterial != nullptr)
 		{
 			CheckerMaterial->SetScalarParameterValue("Density", CheckerDensity);
-			CheckerMaterial->SetScalarParameterValue("UVChannel", (float)UVChannel);
+			CheckerMaterial->SetScalarParameterValue("UVChannel", static_cast<float>(UVChannelNamesList.IndexOfByKey(UVChannel)));
 		}
 	}
 }
@@ -38,7 +43,7 @@ void UExistingMeshMaterialProperties::UpdateMaterials()
 	if (CheckerMaterial != nullptr)
 	{
 		CheckerMaterial->SetScalarParameterValue("Density", CheckerDensity);
-		CheckerMaterial->SetScalarParameterValue("UVChannel", (float)UVChannel);
+		CheckerMaterial->SetScalarParameterValue("UVChannel", static_cast<float>(UVChannelNamesList.IndexOfByKey(UVChannel)));
 	}
 }
 
@@ -56,4 +61,13 @@ UMaterialInterface* UExistingMeshMaterialProperties::GetActiveOverrideMaterial()
 	return nullptr;
 }
 
+void UExistingMeshMaterialProperties::UpdateUVChannels(int32 UVChannelIndex, const TArray<FString>& UVChannelNames, bool bUpdateSelection)
+{
+	UVChannelNamesList = UVChannelNames;
+	if (bUpdateSelection)
+	{
+		UVChannel = 0 <= UVChannelIndex && UVChannelIndex < UVChannelNames.Num() ? UVChannelNames[UVChannelIndex] : TEXT("");
+	}
+}
+
 #undef LOCTEXT_NAMESPACE
diff --git a/Engine/Plugins/Runtime/MeshModelingToolset/Source/MeshModelingTools/Public/Properties/MeshMaterialProperties.h b/Engine/Plugins/Runtime/MeshModelingToolset/Source/MeshModelingTools/Public/Properties/MeshMaterialProperties.h
--- a/Engine/Plugins/Runtime/MeshModelingToolset/Source/MeshModelingTools/Public/Properties/MeshMaterialProperties.h
+++ b/Engine/Plugins/Runtime/MeshModelingToolset/Source/MeshModelingTools/Public/Properties/MeshMaterialProperties.h
@@ -8,7 +8,7 @@
 #include "MeshMaterialProperties.generated.h"
 
 
-// predeclarations
+// Forward declarations
 class UMaterialInterface;
 class UMaterialInstanceDynamic;
 
@@ -49,7 +49,7 @@ UENUM()
 enum class ESetMeshMaterialMode : uint8
 {
 	/** Input material */
-	KeepOriginal,
+	Original,
 
 	/** Checkerboard material */
 	Checkerboard,
@@ -66,19 +66,30 @@ class MESHMODELINGTOOLS_API UExistingMeshMaterialProperties : public UInteractiv
 public:
 
 	/** Material that will be used on the mesh */
-	UPROPERTY(EditAnywhere, Category = MaterialPreview)
-	ESetMeshMaterialMode MaterialMode = ESetMeshMaterialMode::KeepOriginal;
+	UPROPERTY(EditAnywhere, Category = PreviewMaterial)
+	ESetMeshMaterialMode MaterialMode = ESetMeshMaterialMode::Original;
 
-	UPROPERTY(EditAnywhere, Category = MaterialPreview, meta = (UIMin = "1.0", UIMax = "40.0", ClampMin = "0.01", ClampMax = "1000.0", EditConditionHides, EditCondition = "MaterialMode == ESetMeshMaterialMode::Checkerboard"))
+	/** Number of checkerboard tiles within the 0 to 1 range; only available when Checkerboard is selected as material mode */
+	UPROPERTY(EditAnywhere, Category = PreviewMaterial,
+		meta = (UIMin = "1.0", UIMax = "40.0", ClampMin = "0.01", ClampMax = "1000.0", EditConditionHides, EditCondition = "MaterialMode == ESetMeshMaterialMode::Checkerboard"))
 	float CheckerDensity = 20.0f;
 
-	/** Which UV Channel to visualize with checkerboard material */
-	UPROPERTY(EditAnywhere, Category = MaterialPreview, meta = (DisplayName = "Show UV Channel", EditConditionHides, EditCondition = "MaterialMode == ESetMeshMaterialMode::Checkerboard"))
-	int UVChannel = 0;
-
-	UPROPERTY(EditAnywhere, Category = MaterialPreview, meta = (EditCondition = "MaterialMode == ESetMeshMaterialMode::Override"))
+	/** Material to use instead of the original material; only available when Override is selected as material mode */
+	UPROPERTY(EditAnywhere, Category = PreviewMaterial, meta = (EditConditionHides, EditCondition = "MaterialMode == ESetMeshMaterialMode::Override"))
 	TObjectPtr<UMaterialInterface> OverrideMaterial = nullptr;
 
+	/** Which UV channel to use for visualizing the checkerboard material on the mesh; note that this does not affect the preview layout */
+	UPROPERTY(EditAnywhere, Category = PreviewMaterial,
+		meta = (DisplayName = "Preview UV Channel", GetOptions = GetUVChannelNamesFunc, EditConditionHides, EditCondition =
+			"MaterialMode == ESetMeshMaterialMode::Checkerboard", NoResetToDefault))
+	FString UVChannel;
+
+	UPROPERTY(meta = (TransientToolProperty))
+	TArray<FString> UVChannelNamesList;
+
+	UFUNCTION()
+	const TArray<FString>& GetUVChannelNamesFunc() const;
+
 	UPROPERTY(meta = (TransientToolProperty))
 	TObjectPtr<UMaterialInstanceDynamic> CheckerMaterial = nullptr;
 
@@ -89,6 +100,8 @@ public:
 
 	void UpdateMaterials();
 	UMaterialInterface* GetActiveOverrideMaterial() const;
+	
+	void UpdateUVChannels(int32 UVChannelIndex, const TArray<FString>& UVChannelNames, bool bUpdateSelection = true);
 };
 
 UENUM()
diff --git a/Engine/Plugins/Runtime/MeshModelingToolset/Source/ModelingComponents/Private/Drawing/UVLayoutPreview.cpp b/Engine/Plugins/Runtime/MeshModelingToolset/Source/ModelingComponents/Private/Drawing/UVLayoutPreview.cpp
--- a/Engine/Plugins/Runtime/MeshModelingToolset/Source/ModelingComponents/Private/Drawing/UVLayoutPreview.cpp
+++ b/Engine/Plugins/Runtime/MeshModelingToolset/Source/ModelingComponents/Private/Drawing/UVLayoutPreview.cpp
@@ -22,7 +22,7 @@ void UUVLayoutPreview::CreateInWorld(UWorld* World)
 	TriangleComponent->RegisterComponent();
 
 	Settings = NewObject<UUVLayoutPreviewProperties>(this);
-	Settings->WatchProperty(Settings->bVisible, [this](bool) { bSettingsModified = true; });
+	Settings->WatchProperty(Settings->bEnabled, [this](bool) { bSettingsModified = true; });
 	Settings->WatchProperty(Settings->bShowWireframe, [this](bool) { bSettingsModified = true; });
 	//Settings->WatchProperty(Settings->bWireframe, [this](bool) { bSettingsModified = true; });
 	bSettingsModified = true;
@@ -82,7 +82,7 @@ void UUVLayoutPreview::Render(IToolsContextRenderAPI* RenderAPI)
 
 	RecalculatePosition();
 
-	if (Settings->bVisible)
+	if (Settings->bEnabled)
 	{
 		float ScaleFactor = GetCurrentScale();
 		FVector Origin = (FVector)CurrentWorldFrame.Origin;
@@ -106,7 +106,7 @@ void UUVLayoutPreview::OnTick(float DeltaTime)
 {
 	if (bSettingsModified)
 	{
-		SetVisible(Settings->bVisible);
+		SetVisible(Settings->bEnabled);
 
 		PreviewMesh->EnableWireframe(Settings->bShowWireframe);
 
@@ -117,7 +117,7 @@ void UUVLayoutPreview::OnTick(float DeltaTime)
 
 float UUVLayoutPreview::GetCurrentScale()
 {
-	return Settings->ScaleFactor * SourceObjectWorldBounds.Height();
+	return Settings->Scale * SourceObjectWorldBounds.Height();
 }
 
 
@@ -203,15 +203,13 @@ void UUVLayoutPreview::RecalculatePosition()
 
 	double UseScale = GetCurrentScale();
 
-	double ShiftRight = Settings->Shift.X * (ProjectedBounds.Max.X + (ProjectedBounds.Width() * 0.1));
-	//double ShiftRight = Settings->Shift.X * (0.5 * SourceObjectWorldBounds.DiagonalLength());
-	if (Settings->WhichSide == EUVLayoutPreviewSide::Left)
+	double ShiftRight = Settings->Offset.X * (ProjectedBounds.Max.X + (ProjectedBounds.Width() * 0.1));
+	if (Settings->Side == EUVLayoutPreviewSide::Left)
 	{
-		ShiftRight = ProjectedBounds.Min.X - Settings->Shift.X * (UseScale + (ProjectedBounds.Width() * 0.1));
+		ShiftRight = ProjectedBounds.Min.X - Settings->Offset.X * (UseScale + (ProjectedBounds.Width() * 0.1));
 	}
 
-	//double ShiftUp = ProjectedBounds.Max.Y + Settings->ShiftY * (ProjectedBounds.Height() + Settings->ScaleFactor);
-	double ShiftUp = Settings->Shift.Y * UseScale;
+	double ShiftUp = Settings->Offset.Y * UseScale;
 	ObjFrame.Origin += ShiftRight * ObjFrame.X() - ShiftUp * ObjFrame.Y();
 
 	CurrentWorldFrame = ObjFrame;
diff --git a/Engine/Plugins/Runtime/MeshModelingToolset/Source/ModelingComponents/Public/Drawing/UVLayoutPreview.h b/Engine/Plugins/Runtime/MeshModelingToolset/Source/ModelingComponents/Public/Drawing/UVLayoutPreview.h
--- a/Engine/Plugins/Runtime/MeshModelingToolset/Source/ModelingComponents/Public/Drawing/UVLayoutPreview.h
+++ b/Engine/Plugins/Runtime/MeshModelingToolset/Source/ModelingComponents/Public/Drawing/UVLayoutPreview.h
@@ -16,8 +16,11 @@ PREDECLARE_USE_GEOMETRY_CLASS(FDynamicMesh3);
 UENUM()
 enum class EUVLayoutPreviewSide
 {
-	Left = 0,
-	Right = 1
+	/** On the left side of the object center relative to the camera */
+	Left,
+
+	/** On the right side of the object center relative to the camera */
+	Right
 };
 
 
@@ -29,24 +32,26 @@ class MODELINGCOMPONENTS_API UUVLayoutPreviewProperties : public UInteractiveToo
 {
 	GENERATED_BODY()
 public:
-	/** Should be UV Layout be shown */
-	UPROPERTY(EditAnywhere, Category = UVLayoutPreview)
-	bool bVisible = true;
+	/** Enable the preview UV layout */
+	UPROPERTY(EditAnywhere, Category = "Preview UV Layout")
+	bool bEnabled = true;
 
-	/** World-space scaling factor on the UV Layout */
-	UPROPERTY(EditAnywhere, Category = UVLayoutPreview, meta = (UIMin = "0.1", UIMax = "10.0", ClampMin = "0.0001", ClampMax = "1000"))
-	float ScaleFactor = 1.0;
+	/** Which side of the selected object the preview UV layout is shown */
+	UPROPERTY(EditAnywhere, Category = "Preview UV Layout", meta = (EditCondition = "bEnabled"))
+	EUVLayoutPreviewSide Side = EUVLayoutPreviewSide::Right;
 
-	/** Where should the UV layout be positioned relative to the target object, relative to camera */
-	UPROPERTY(EditAnywhere, Category = UVLayoutPreview)
-	EUVLayoutPreviewSide WhichSide = EUVLayoutPreviewSide::Right;
+	/** Uniform scaling factor */
+	UPROPERTY(EditAnywhere, Category = "Preview UV Layout",
+		meta = (UIMin = "0.1", UIMax = "10.0", ClampMin = "0.0001", ClampMax = "1000", EditCondition = "bEnabled"))
+	float Scale = 1.0;
 
-	/** If true, wireframe is shown for the UV layout */
-	UPROPERTY(EditAnywhere, Category = UVLayoutPreview)
-	bool bShowWireframe = true;
+	/** Offset relative to the center of the selected object */
+	UPROPERTY(EditAnywhere, Category = "Preview UV Layout", meta = (EditCondition = "bEnabled"))
+	FVector2D Offset = FVector2D(1.0, 0.5);
 
-	UPROPERTY(EditAnywhere, Category = UVLayoutPreview)
-	FVector2D Shift = FVector2D(1.0, 0.5);
+	/** Show wireframe mesh in the preview UV layout */
+	UPROPERTY(EditAnywhere, Category = "Preview UV Layout", meta = (EditCondition = "bEnabled"))
+	bool bShowWireframe = true;
 };
 
 
@@ -62,7 +67,7 @@ class MODELINGCOMPONENTS_API UUVLayoutPreview : public UObject
 	GENERATED_BODY()
 
 public:
-	virtual ~UUVLayoutPreview();
+	virtual ~UUVLayoutPreview() override;
 
 
 	/**
@@ -98,12 +103,12 @@ public:
 	void UpdateUVMesh(const FDynamicMesh3* SourceMesh, int32 SourceUVLayer = 0);
 
 	/**
-	 * Tick the UV Layout Preview, allowing it to upodate various settings
+	 * Tick the UV Layout Preview, allowing it to update various settings
 	 */
 	void OnTick(float DeltaTime);
 
 	/**
-	 * Render the UV Layout Preview, allowing it to upodate various settings
+	 * Render the UV Layout Preview, allowing it to update various settings
 	 */
 	void Render(IToolsContextRenderAPI* RenderAPI);
 
diff --git a/Engine/Plugins/Runtime/MeshModelingToolset/Source/ModelingOperators/Private/ParameterizationOps/UVLayoutOp.cpp b/Engine/Plugins/Runtime/MeshModelingToolset/Source/ModelingOperators/Private/ParameterizationOps/UVLayoutOp.cpp
--- a/Engine/Plugins/Runtime/MeshModelingToolset/Source/ModelingOperators/Private/ParameterizationOps/UVLayoutOp.cpp
+++ b/Engine/Plugins/Runtime/MeshModelingToolset/Source/ModelingOperators/Private/ParameterizationOps/UVLayoutOp.cpp
@@ -114,8 +114,8 @@ TUniquePtr<FDynamicMeshOperator> UUVLayoutOperatorFactory::MakeNewOperator()
 	Op->UVLayerIndex = GetSelectedUVChannel();
 	Op->TextureResolution = Settings->TextureResolution;
 	Op->bAllowFlips = Settings->bAllowFlips;
-	Op->UVScaleFactor = Settings->UVScaleFactor;
-	Op->UVTranslation = FVector2f(Settings->UVTranslate);
+	Op->UVScaleFactor = Settings->Scale;
+	Op->UVTranslation = FVector2f(Settings->Translation);
 	Op->SetTransform(TargetTransform);
 
 	return Op;
diff --git a/Engine/Plugins/Runtime/MeshModelingToolset/Source/ModelingOperators/Public/Properties/UVLayoutProperties.h b/Engine/Plugins/Runtime/MeshModelingToolset/Source/ModelingOperators/Public/Properties/UVLayoutProperties.h
--- a/Engine/Plugins/Runtime/MeshModelingToolset/Source/ModelingOperators/Public/Properties/UVLayoutProperties.h
+++ b/Engine/Plugins/Runtime/MeshModelingToolset/Source/ModelingOperators/Public/Properties/UVLayoutProperties.h
@@ -3,7 +3,6 @@
 #pragma once
 
 #include "CoreMinimal.h"
-
 #include "InteractiveTool.h"
 
 #include "UVLayoutProperties.generated.h"
@@ -14,11 +13,11 @@
 UENUM()
 enum class EUVLayoutType
 {
-	/** Apply Translation and Scale to existing UVs */
+	/** Apply Scale and Translation properties to all UV values */
 	Transform,
-	/** Uniformly scale/translate UV Islands into the standard UV Unit Square */
+	/** Uniformly scale and translate each UV island individually to pack it into the unit square, i.e. fit between 0 and 1 with overlap */
 	Stack,
-	/** Repack all UV Islands to fit inside standard UV unit square */
+	/** Uniformly scale and translate UV islands collectively to pack them into the unit square, i.e. fit between 0 and 1 with no overlap */
 	Repack
 };
 
@@ -32,24 +31,24 @@ class MODELINGOPERATORS_API UUVLayoutProperties : public UInteractiveToolPropert
 	GENERATED_BODY()
 
 public:
-	/** Type of transformation to apply to input UV islands */
-	UPROPERTY(EditAnywhere, Category = UVLayout)
+	/** Type of layout applied to input UVs */
+	UPROPERTY(EditAnywhere, Category = "UV Layout")
 	EUVLayoutType LayoutType = EUVLayoutType::Repack;
 
-	/** Expected resolution of output textures; controls spacing left between charts */
-	UPROPERTY(EditAnywhere, Category = UVLayout, meta = (UIMin = "64", UIMax = "2048", ClampMin = "2", ClampMax = "4096"))
+	/** Expected resolution of the output textures; this controls spacing left between UV islands to avoid interpolation artifacts */
+	UPROPERTY(EditAnywhere, Category = "UV Layout", meta = (UIMin = "64", UIMax = "2048", ClampMin = "2", ClampMax = "4096"))
 	int TextureResolution = 1024;
 
-	/** Apply this uniform scaling to the UVs after any layout recalculation */
-	UPROPERTY(EditAnywhere, Category = UVLayout, meta = (UIMin = "0.1", UIMax = "5.0", ClampMin = "0.0001", ClampMax = "10000") )
-	float UVScaleFactor = 1;
+	/** Uniform scale applied to UVs after packing */
+	UPROPERTY(EditAnywhere, Category = "UV Layout", meta = (UIMin = "0.1", UIMax = "5.0", ClampMin = "0.0001", ClampMax = "10000") )
+	float Scale = 1;
 
-	/** Apply this 2D translation to the UVs after any layout recalculation, and after scaling */
-	UPROPERTY(EditAnywhere, Category = UVLayout)
-	FVector2D UVTranslate = FVector2D(0,0);
+	/** Translation applied to UVs after packing, and after scaling */
+	UPROPERTY(EditAnywhere, Category = "UV Layout")
+	FVector2D Translation = FVector2D(0,0);
 
-	/** Allow the packer to flip the orientation of UV islands if it save space. May cause problems for downstream operations, not recommended. */
-	UPROPERTY(EditAnywhere, Category = UVLayout, AdvancedDisplay)
+	/** Allow the Repack layout type to flip the orientation of UV islands to save space. Note that this may cause problems for downstream operations, and therefore is disabled by default. */
+	UPROPERTY(EditAnywhere, Category = "UV Layout", meta = (EditCondition = "LayoutType == EUVLayoutType::Repack"))
 	bool bAllowFlips = false;
 
-};
\ No newline at end of file
+};
