ModelingTools: Add details customization of EBakeMapType properties for bake tools.

#rnx

diff --git a/Engine/Plugins/Editor/ModelingToolsEditorMode/Source/ModelingToolsEditorMode/Private/DetailsCustomizations/BakeMeshAttributeToolCustomizations.cpp b/Engine/Plugins/Editor/ModelingToolsEditorMode/Source/ModelingToolsEditorMode/Private/DetailsCustomizations/BakeMeshAttributeToolCustomizations.cpp
new file mode 100644
--- /dev/null
+++ b/Engine/Plugins/Editor/ModelingToolsEditorMode/Source/ModelingToolsEditorMode/Private/DetailsCustomizations/BakeMeshAttributeToolCustomizations.cpp
@@ -0,0 +1,392 @@
+// Copyright Epic Games, Inc. All Rights Reserved.
+
+#include "DetailsCustomizations/BakeMeshAttributeToolCustomizations.h"
+#include "DetailCategoryBuilder.h"
+#include "DetailLayoutBuilder.h"
+#include "DetailWidgetRow.h"
+#include "PropertyCustomizationHelpers.h"
+#include "Widgets/Input/SComboButton.h"
+#include "IDetailChildrenBuilder.h"
+#include "Internationalization/BreakIterator.h"
+#include "Framework/MultiBox/MultiBoxBuilder.h"
+#include "IContentBrowserSingleton.h"
+#include "SAssetView.h"
+
+#include "ModelingWidgets/SDynamicNumericEntry.h"
+#include "ModelingWidgets/ModelingCustomizationUtil.h"
+
+#include "ModelingToolsEditorModeSettings.h"
+
+#include "BakeMeshAttributeMapsTool.h"
+#include "BakeMultiMeshAttributeMapsTool.h"
+#include "BakeMeshAttributeVertexTool.h"
+
+#define LOCTEXT_NAMESPACE "BakeMeshAttributeToolCustomizations"
+
+//
+// Bake Details locals
+//
+namespace BakeCustomizationLocals
+{
+	/** EBakeMapType UI sections */
+	enum class EBakeMapTypeSections
+	{
+		Surface,
+		Lighting,
+		Material,
+		Mask,
+		MaxSections // Must be last
+	};
+	
+	/** EBakeMapType UI section header names */
+	static const TArray<FText> BakeMapTypeSectionNames =
+		{
+			LOCTEXT("BakeMapTypeSurfaceSection", "Surface"),
+			LOCTEXT("BakeMapTypeLightingSection", "Lighting"),
+			LOCTEXT("BakeMapTypeMaterialSection", "Material"),
+			LOCTEXT("BakeMapTypeMasksSection", "IDs & Masks")
+		};
+
+	/** EBakeMapType UI section content */
+	static const TArray<TArray<EBakeMapType>> BakeMapTypeSectionData =
+		{
+			// Surface
+			{
+				EBakeMapType::TangentSpaceNormal,
+				EBakeMapType::ObjectSpaceNormal,
+				EBakeMapType::FaceNormal,
+				EBakeMapType::BentNormal,
+				EBakeMapType::Position,
+				EBakeMapType::Curvature
+			},
+			// Lighting
+			{
+				EBakeMapType::AmbientOcclusion
+			},
+			// Material
+			{
+				EBakeMapType::Texture,
+				EBakeMapType::MultiTexture
+			},
+			// Mask
+			{
+				EBakeMapType::VertexColor,
+				EBakeMapType::MaterialID
+			}
+		};
+
+	/**
+	 * Return an array of names from the metadata flag "ValidEnumValues".
+	 * Sourced from privately defined PropertyEditorHelpers::GetValidEnumsFromPropertyOverride().
+	 *
+	 * @param Property Property to inspect
+	 * @param InEnum UEnum class that this property represents
+	 * @return array of FName for each valid enum value
+	 */
+	TArray<FName> GetValidEnumValues(const FProperty* Property, const UEnum* InEnum)
+	{
+		TArray<FName> ValidEnumValues;
+		static const FName ValidEnumValuesName(TEXT("ValidEnumValues"));
+		if(Property->HasMetaData(ValidEnumValuesName))
+		{
+			TArray<FString> ValidEnumValuesAsString;
+			Property->GetMetaData(ValidEnumValuesName).ParseIntoArray(ValidEnumValuesAsString, TEXT(","));
+			for(FString& Value : ValidEnumValuesAsString)
+			{
+				Value.TrimStartInline();
+				ValidEnumValues.Add(*InEnum->GenerateFullEnumName(*Value));
+			}
+		}
+		return ValidEnumValues;
+	}
+
+	/**
+	 * Return the UEnum associated with the BitmaskEnum metadata flag.
+	 * Sourced from privately defined SPropertyEditorNumeric::Construct().
+	 *
+	 * @param Property the property to query for the BitmaskEnum metadata
+	 * @return the UEnum defined by the BitmaskEnum metadata. nullptr if undefined or cannot be found.
+	 */
+	UEnum* GetBitmaskEnum(const FProperty* Property)
+	{
+		UEnum* BitmaskEnum = nullptr;
+		static const FName BitmaskEnumFlagName(TEXT("BitmaskEnum"));
+		const FString& BitmaskEnumName = Property->GetMetaData(BitmaskEnumFlagName);
+		if (!BitmaskEnumName.IsEmpty())
+		{
+			BitmaskEnum = FindObject<UEnum>(ANY_PACKAGE, *BitmaskEnumName);
+		}
+		return BitmaskEnum;
+	}
+
+	/**
+	 * Utility struct that holds bitmask flag UI data.
+	 */
+	struct FBitmaskFlagInfo
+	{
+		int32 Value;
+		FText DisplayName;
+		FText ToolTipText;
+	};
+
+	/**
+	 * Return an array of FBitmaskFlagInfo for each valid enum value in the enum
+	 * defined by a property's BitmaskEnum metadata.
+	 * Sourced from privately defined SPropertyEditorNumeric::Construct()
+	 *
+	 * @param Property the property to query
+	 * @return an array of FBitmaskFlagInfo for each valid enum value
+	 */
+	TArray<FBitmaskFlagInfo> GetBitmaskEnumFlags(const FProperty* Property)
+	{
+		constexpr int32 BitmaskBitCount = sizeof(int) << 3;
+		
+		TArray<FBitmaskFlagInfo> Result;
+		Result.Empty(BitmaskBitCount);
+		const UEnum* BitmaskEnum = GetBitmaskEnum(Property);
+		if (!BitmaskEnum)
+		{
+			return Result;
+		}
+
+		auto AddNewBitmaskFlag = [&Result, BitmaskEnum](const int32 EnumIndex, const int64& EnumValue)
+		{
+			Result.Emplace();
+			FBitmaskFlagInfo* BitmaskFlag = &Result.Last();
+
+			BitmaskFlag->Value = EnumValue;
+			BitmaskFlag->DisplayName = BitmaskEnum->GetDisplayNameTextByIndex(EnumIndex);
+			BitmaskFlag->ToolTipText = BitmaskEnum->GetToolTipTextByIndex(EnumIndex);
+			if (BitmaskFlag->ToolTipText.IsEmpty())
+			{
+				BitmaskFlag->ToolTipText = FText::Format(LOCTEXT("BitmaskDefaultFlagToolTipText", "Toggle {0} on/off"), BitmaskFlag->DisplayName);
+			}
+		};
+
+		static const FName UseEnumValuesAsMaskValuesName(TEXT("UseEnumValuesAsMaskValuesInEditor"));
+		const bool bUseEnumValuesAsMaskValues = BitmaskEnum->GetBoolMetaData(UseEnumValuesAsMaskValuesName);
+		const TArray<FName> AllowedPropertyEnums = GetValidEnumValues(Property, BitmaskEnum);
+		// Note: This loop doesn't include (BitflagsEnum->NumEnums() - 1) in order to skip the implicit "MAX"
+		// value that gets added to the enum type at compile time.
+		for (int32 BitmaskEnumIndex = 0; BitmaskEnumIndex < BitmaskEnum->NumEnums() - 1; ++BitmaskEnumIndex)
+		{
+			const int64 EnumValue = BitmaskEnum->GetValueByIndex(BitmaskEnumIndex);
+			bool bShouldBeHidden = BitmaskEnum->HasMetaData(TEXT("Hidden"), BitmaskEnumIndex);
+			if (!bShouldBeHidden && AllowedPropertyEnums.Num() > 0)
+			{
+				bShouldBeHidden = AllowedPropertyEnums.Find(BitmaskEnum->GetNameByIndex(BitmaskEnumIndex)) == INDEX_NONE;
+			}
+			if (EnumValue >= 0 && !bShouldBeHidden)
+			{
+				if (bUseEnumValuesAsMaskValues)
+				{
+					if (EnumValue < MAX_int64 && FMath::IsPowerOfTwo(EnumValue))
+					{
+						AddNewBitmaskFlag(BitmaskEnumIndex, EnumValue);
+					}
+				}
+				else if (EnumValue < BitmaskBitCount)
+				{
+					AddNewBitmaskFlag(BitmaskEnumIndex, static_cast<int64>(1) << EnumValue);
+				}
+			}
+		}
+		return Result;
+	}
+
+	/**
+	 * Customize the EBakeMapTypes bitmask enum with predefined section headers.
+	 *
+	 * This function extends the existing default widget rather than replacing it, so
+	 * the functionality is tied to its definition in SPropertyEditorNumeric. Should
+	 * the widget hierarchy of SPropertyEditorNumeric change, this function should fail
+	 * gracefully.
+	 *
+	 * @param DetailBuilder the detail layout builder
+	 * @param MapTypesHandle handle to a bitmask enum property for EBakeMapTypes
+	 * @param bMultiSelect if true, allow multiple enum values to be selected
+	 */
+	void CustomizeMapTypesEnum(IDetailLayoutBuilder& DetailBuilder, const TSharedPtr<IPropertyHandle>& MapTypesHandle, const bool bMultiSelect)
+	{
+		if (!ensure(MapTypesHandle->IsValidHandle()))
+		{
+			return;
+		}
+		FProperty* MapTypesProperty = MapTypesHandle->GetProperty();
+
+		IDetailPropertyRow* DetailRow = DetailBuilder.EditDefaultProperty(MapTypesHandle);
+		TSharedPtr<SWidget> NameWidget, ValueWidget;
+		DetailRow->GetDefaultWidgets(NameWidget, ValueWidget);
+
+		// Find first SComboButton in hierarchy
+		TSharedPtr<SComboButton> ComboWidget;
+		UE::ModelingUI::ProcessChildWidgetsByType(
+			ValueWidget->AsShared(),
+			TEXT("SComboButton"),
+			[&ComboWidget](TSharedRef<SWidget> Widget)->bool
+		{
+			ComboWidget = StaticCastSharedPtr<SComboButton>(TSharedPtr<SWidget>(Widget));
+			// Stop processing after first occurrence
+			return false;
+		});
+		
+		if (!ComboWidget.IsValid())
+		{
+			return;
+		}
+
+		// Override the ComboButton menu builder with customized sections & entries
+		ComboWidget->SetOnGetMenuContent(FOnGetContent::CreateLambda([MapTypesProperty, DetailRow, bMultiSelect]()
+		{
+			FMenuBuilder MenuBuilder(!bMultiSelect, nullptr);
+
+			TArray<FBitmaskFlagInfo> BitmaskFlags = GetBitmaskEnumFlags(MapTypesProperty);
+			TMap<int64, FBitmaskFlagInfo> BitmaskFlagMap;
+			const int NumFlags = BitmaskFlags.Num();
+			for (int Idx = 0; Idx < NumFlags; ++Idx)
+			{
+				BitmaskFlagMap.Emplace(BitmaskFlags[Idx].Value, BitmaskFlags[Idx]);
+			}
+
+			auto AddMenuEntry = [DetailRow, &MenuBuilder, bMultiSelect](const FBitmaskFlagInfo& BitmaskFlag)
+			{
+				MenuBuilder.AddMenuEntry(
+					BitmaskFlag.DisplayName,
+					BitmaskFlag.ToolTipText,
+					FSlateIcon(),
+					FUIAction
+					(
+						FExecuteAction::CreateLambda([DetailRow, BitmaskFlag, bMultiSelect]()
+						{
+							const TSharedPtr<IPropertyHandle> Property = DetailRow->GetPropertyHandle();
+							if (Property.IsValid())
+							{
+								int32 Value;
+								if (Property->GetValue(Value) == FPropertyAccess::Success)
+								{
+									bMultiSelect ? Property->SetValue(Value ^ BitmaskFlag.Value) : Property->SetValue(BitmaskFlag.Value);
+								}
+							}
+						}),
+						FCanExecuteAction(),
+						FIsActionChecked::CreateLambda([DetailRow, BitmaskFlag]() -> bool
+						{
+							const TSharedPtr<IPropertyHandle> Property = DetailRow->GetPropertyHandle();
+							if (Property.IsValid())
+							{
+								int32 Value;
+								if (Property->GetValue(Value) == FPropertyAccess::Success)
+								{
+									return (Value & BitmaskFlag.Value) != 0;
+								}
+							}
+							return false;
+						})
+					),
+					NAME_None,
+					EUserInterfaceActionType::Check);
+			};
+
+			// Explicitly handle EBakeMapTypes::None to be at the head of the list if valid.
+			if (const FBitmaskFlagInfo* BitmaskFlag = BitmaskFlagMap.Find(static_cast<int64>(EBakeMapType::None)))
+			{
+				AddMenuEntry(*BitmaskFlag);
+			}
+
+			constexpr int NumSections = static_cast<int>(EBakeMapTypeSections::MaxSections);
+			for (int SectionIdx = 0; SectionIdx < NumSections; ++SectionIdx)
+			{
+				MenuBuilder.BeginSection(NAME_None, BakeMapTypeSectionNames[SectionIdx]);
+				const int NumSectionData = BakeMapTypeSectionData[SectionIdx].Num();
+				for (int DataIdx = 0; DataIdx < NumSectionData; ++DataIdx)
+				{
+					const EBakeMapType DataType = BakeMapTypeSectionData[SectionIdx][DataIdx];
+					if (const FBitmaskFlagInfo* BitmaskFlag = BitmaskFlagMap.Find(static_cast<int64>(DataType)))
+					{
+						AddMenuEntry(*BitmaskFlag);
+					}
+				}
+				MenuBuilder.EndSection();
+			}
+			return MenuBuilder.MakeWidget();
+		}));
+
+		// Assign default name and modified value widget to respective content widgets.
+		DetailRow->CustomWidget()
+			.NameContent()
+			[
+				NameWidget->AsShared()
+			]
+			.ValueContent()
+			[
+				ValueWidget->AsShared()
+			];
+	}
+};
+
+
+//
+// Mesh Map Bake
+//
+
+
+TSharedRef<IDetailCustomization> FBakeMeshAttributeMapsToolDetails::MakeInstance()
+{
+	return MakeShareable(new FBakeMeshAttributeMapsToolDetails);
+}
+
+
+void FBakeMeshAttributeMapsToolDetails::CustomizeDetails(IDetailLayoutBuilder& DetailBuilder)
+{
+	const TSharedPtr<IPropertyHandle> MapTypesHandle = DetailBuilder.GetProperty(GET_MEMBER_NAME_CHECKED(UBakeMeshAttributeMapsToolProperties, MapTypes), UBakeMeshAttributeMapsToolProperties::StaticClass());
+	BakeCustomizationLocals::CustomizeMapTypesEnum(DetailBuilder, MapTypesHandle, true);
+}
+
+
+//
+// MultiMesh Map Bake
+//
+
+
+TSharedRef<IDetailCustomization> FBakeMultiMeshAttributeMapsToolDetails::MakeInstance()
+{
+	return MakeShareable(new FBakeMultiMeshAttributeMapsToolDetails);
+}
+
+
+void FBakeMultiMeshAttributeMapsToolDetails::CustomizeDetails(IDetailLayoutBuilder& DetailBuilder)
+{
+	const TSharedPtr<IPropertyHandle> MapTypesHandle = DetailBuilder.GetProperty(GET_MEMBER_NAME_CHECKED(UBakeMultiMeshAttributeMapsToolProperties, MapTypes), UBakeMultiMeshAttributeMapsToolProperties::StaticClass());
+	BakeCustomizationLocals::CustomizeMapTypesEnum(DetailBuilder, MapTypesHandle, true);
+}
+
+
+//
+// Mesh Vertex Bake
+//
+
+
+TSharedRef<IDetailCustomization> FBakeMeshAttributeVertexToolDetails::MakeInstance()
+{
+	return MakeShareable(new FBakeMeshAttributeVertexToolDetails);
+}
+
+
+void FBakeMeshAttributeVertexToolDetails::CustomizeDetails(IDetailLayoutBuilder& DetailBuilder)
+{
+	const TSharedPtr<IPropertyHandle> OutputTypeHandle = DetailBuilder.GetProperty(GET_MEMBER_NAME_CHECKED(UBakeMeshAttributeVertexToolProperties, OutputType), UBakeMeshAttributeVertexToolProperties::StaticClass());
+	BakeCustomizationLocals::CustomizeMapTypesEnum(DetailBuilder, OutputTypeHandle, false);
+
+	const TSharedPtr<IPropertyHandle> OutputTypeRHandle = DetailBuilder.GetProperty(GET_MEMBER_NAME_CHECKED(UBakeMeshAttributeVertexToolProperties, OutputTypeR), UBakeMeshAttributeVertexToolProperties::StaticClass());
+	const TSharedPtr<IPropertyHandle> OutputTypeGHandle = DetailBuilder.GetProperty(GET_MEMBER_NAME_CHECKED(UBakeMeshAttributeVertexToolProperties, OutputTypeG), UBakeMeshAttributeVertexToolProperties::StaticClass());
+	const TSharedPtr<IPropertyHandle> OutputTypeBHandle = DetailBuilder.GetProperty(GET_MEMBER_NAME_CHECKED(UBakeMeshAttributeVertexToolProperties, OutputTypeB), UBakeMeshAttributeVertexToolProperties::StaticClass());
+	const TSharedPtr<IPropertyHandle> OutputTypeAHandle = DetailBuilder.GetProperty(GET_MEMBER_NAME_CHECKED(UBakeMeshAttributeVertexToolProperties, OutputTypeA), UBakeMeshAttributeVertexToolProperties::StaticClass());
+	BakeCustomizationLocals::CustomizeMapTypesEnum(DetailBuilder, OutputTypeRHandle, false);
+	BakeCustomizationLocals::CustomizeMapTypesEnum(DetailBuilder, OutputTypeGHandle, false);
+	BakeCustomizationLocals::CustomizeMapTypesEnum(DetailBuilder, OutputTypeBHandle, false);
+	BakeCustomizationLocals::CustomizeMapTypesEnum(DetailBuilder, OutputTypeAHandle, false);
+}
+
+
+#undef LOCTEXT_NAMESPACE
+
diff --git a/Engine/Plugins/Editor/ModelingToolsEditorMode/Source/ModelingToolsEditorMode/Private/ModelingToolsEditorModeModule.cpp b/Engine/Plugins/Editor/ModelingToolsEditorMode/Source/ModelingToolsEditorMode/Private/ModelingToolsEditorModeModule.cpp
--- a/Engine/Plugins/Editor/ModelingToolsEditorMode/Source/ModelingToolsEditorMode/Private/ModelingToolsEditorModeModule.cpp
+++ b/Engine/Plugins/Editor/ModelingToolsEditorMode/Source/ModelingToolsEditorMode/Private/ModelingToolsEditorModeModule.cpp
@@ -14,9 +14,13 @@
 #include "DetailsCustomizations/ModelingToolPropertyCustomizations.h"
 #include "DetailsCustomizations/ModelingToolsBrushSizeCustomization.h"
 #include "DetailsCustomizations/MeshVertexSculptToolCustomizations.h"
+#include "DetailsCustomizations/BakeMeshAttributeToolCustomizations.h"
 
 #include "PropertySets/AxisFilterPropertyType.h"
 #include "MeshVertexSculptTool.h"
+#include "BakeMeshAttributeMapsTool.h"
+#include "BakeMultiMeshAttributeMapsTool.h"
+#include "BakeMeshAttributeVertexTool.h"
 
 
 #define LOCTEXT_NAMESPACE "FModelingToolsEditorModeModule"
@@ -69,6 +73,7 @@ void FModelingToolsEditorModeModule::OnPostEngineInit()
 
 	// Register details view customizations
 	FPropertyEditorModule& PropertyModule = FModuleManager::LoadModuleChecked<FPropertyEditorModule>("PropertyEditor");
+	/// Sculpt
 	PropertyModule.RegisterCustomPropertyTypeLayout("ModelingToolsAxisFilter", FOnGetPropertyTypeCustomizationInstance::CreateStatic(&FModelingToolsAxisFilterCustomization::MakeInstance));
 	PropertiesToUnregisterOnShutdown.Add(FModelingToolsAxisFilter::StaticStruct()->GetFName());
 	PropertyModule.RegisterCustomPropertyTypeLayout("BrushToolRadius", FOnGetPropertyTypeCustomizationInstance::CreateStatic(&FModelingToolsBrushSizeCustomization::MakeInstance));
@@ -79,6 +84,13 @@ void FModelingToolsEditorModeModule::OnPostEngineInit()
 	ClassesToUnregisterOnShutdown.Add(UVertexBrushSculptProperties::StaticClass()->GetFName());
 	PropertyModule.RegisterCustomClassLayout("VertexBrushAlphaProperties", FOnGetDetailCustomizationInstance::CreateStatic(&FVertexBrushAlphaPropertiesDetails::MakeInstance));
 	ClassesToUnregisterOnShutdown.Add(UVertexBrushAlphaProperties::StaticClass()->GetFName());
+	/// Bake
+	PropertyModule.RegisterCustomClassLayout("BakeMeshAttributeMapsToolProperties", FOnGetDetailCustomizationInstance::CreateStatic(&FBakeMeshAttributeMapsToolDetails::MakeInstance));
+	ClassesToUnregisterOnShutdown.Add(UBakeMeshAttributeMapsToolProperties::StaticClass()->GetFName());
+	PropertyModule.RegisterCustomClassLayout("BakeMultiMeshAttributeMapsToolProperties", FOnGetDetailCustomizationInstance::CreateStatic(&FBakeMultiMeshAttributeMapsToolDetails::MakeInstance));
+	ClassesToUnregisterOnShutdown.Add(UBakeMultiMeshAttributeMapsToolProperties::StaticClass()->GetFName());
+	PropertyModule.RegisterCustomClassLayout("BakeMeshAttributeVertexToolProperties", FOnGetDetailCustomizationInstance::CreateStatic(&FBakeMeshAttributeVertexToolDetails::MakeInstance));
+	ClassesToUnregisterOnShutdown.Add(UBakeMeshAttributeVertexToolProperties::StaticClass()->GetFName());
 }
 
 #undef LOCTEXT_NAMESPACE
diff --git a/Engine/Plugins/Editor/ModelingToolsEditorMode/Source/ModelingToolsEditorMode/Private/ModelingToolsEditorModeToolkit.cpp b/Engine/Plugins/Editor/ModelingToolsEditorMode/Source/ModelingToolsEditorMode/Private/ModelingToolsEditorModeToolkit.cpp
--- a/Engine/Plugins/Editor/ModelingToolsEditorMode/Source/ModelingToolsEditorMode/Private/ModelingToolsEditorModeToolkit.cpp
+++ b/Engine/Plugins/Editor/ModelingToolsEditorMode/Source/ModelingToolsEditorMode/Private/ModelingToolsEditorModeToolkit.cpp
@@ -24,6 +24,7 @@
 #include "Widgets/Layout/SExpandableArea.h"
 #include "Framework/MultiBox/SToolBarButtonBlock.h"
 #include "Internationalization/Text.h"
+#include "ModelingWidgets/ModelingCustomizationUtil.h"
 
 // for Tool Extensions
 #include "Features/IModularFeatures.h"
@@ -831,36 +832,14 @@ void FModelingToolsEditorModeToolkit::InvokeUI()
 			{
 				auto FindFirstChildWidget = [](const TSharedPtr<SWidget>& Widget, const FString& WidgetType)
 				{
-					auto FindFirstChildWidgetImpl = [](const TSharedPtr<SWidget>& Widget, const FString& WidgetType, auto& FindRef) -> TSharedPtr<SWidget>
+					TSharedPtr<SWidget> Result;
+					UE::ModelingUI::ProcessChildWidgetsByType(Widget->AsShared(), WidgetType,[&Result](TSharedRef<SWidget> Widget)
 					{
-						TSharedPtr<SWidget> Result;
-						if (Widget.IsValid())
-						{
-							FChildren* Children = Widget->GetChildren();
-							const int32 NumChild = Children ? Children->NumSlot() : 0;
-							for (int32 ChildIdx = 0; ChildIdx < NumChild; ++ChildIdx)
-							{
-								const TSharedRef<SWidget> ChildWidgetRef = Children->GetChildAt(ChildIdx);
-								TSharedPtr<SWidget> ChildWidgetPtr(ChildWidgetRef);
-								if (ChildWidgetPtr.IsValid())
-								{
-									if (ChildWidgetPtr->GetTypeAsString().Compare(WidgetType) == 0)
-									{
-										Result = ChildWidgetPtr;
-										break;
-									}
-
-									Result = FindRef(ChildWidgetPtr, WidgetType, FindRef);
-									if (Result.IsValid())
-									{
-										break;
-									}
-								}
-							}
-						}
-						return Result;
-					};
-					return FindFirstChildWidgetImpl(Widget, WidgetType, FindFirstChildWidgetImpl);
+						Result = TSharedPtr<SWidget>(Widget);
+						// Stop processing after first occurrence
+						return false;
+					});
+					return Result;
 				};
 
 				TSharedPtr<SWidget> PanelWidget = FindFirstChildWidget(ToolbarRow.ToolbarWidget, TEXT("SUniformWrapPanel"));
diff --git a/Engine/Plugins/Editor/ModelingToolsEditorMode/Source/ModelingToolsEditorMode/Public/DetailsCustomizations/BakeMeshAttributeToolCustomizations.h b/Engine/Plugins/Editor/ModelingToolsEditorMode/Source/ModelingToolsEditorMode/Public/DetailsCustomizations/BakeMeshAttributeToolCustomizations.h
new file mode 100644
--- /dev/null
+++ b/Engine/Plugins/Editor/ModelingToolsEditorMode/Source/ModelingToolsEditorMode/Public/DetailsCustomizations/BakeMeshAttributeToolCustomizations.h
@@ -0,0 +1,31 @@
+// Copyright Epic Games, Inc. All Rights Reserved.
+
+#pragma once
+
+#include "IDetailCustomization.h"
+#include "Templates/SharedPointer.h"
+
+// Customization for UBakeMeshAttributeMapsToolProperties
+class FBakeMeshAttributeMapsToolDetails : public IDetailCustomization
+{
+public:
+	static TSharedRef<IDetailCustomization> MakeInstance();
+	virtual void CustomizeDetails(IDetailLayoutBuilder& DetailBuilder) override;
+};
+
+// Customization for UBakeMultiMeshAttributeMapsToolProperties
+class FBakeMultiMeshAttributeMapsToolDetails : public IDetailCustomization
+{
+public:
+	static TSharedRef<IDetailCustomization> MakeInstance();
+	virtual void CustomizeDetails(IDetailLayoutBuilder& DetailBuilder) override;
+};
+
+// Customization for UBakeMeshAttributeVertexToolProperties
+class FBakeMeshAttributeVertexToolDetails : public IDetailCustomization
+{
+public:
+	static TSharedRef<IDetailCustomization> MakeInstance();
+	virtual void CustomizeDetails(IDetailLayoutBuilder& DetailBuilder) override;
+};
+
diff --git a/Engine/Plugins/Experimental/MeshModelingToolsetExp/Source/MeshModelingToolsExp/Private/BakeMeshAttributeVertexTool.cpp b/Engine/Plugins/Experimental/MeshModelingToolsetExp/Source/MeshModelingToolsExp/Private/BakeMeshAttributeVertexTool.cpp
--- a/Engine/Plugins/Experimental/MeshModelingToolsetExp/Source/MeshModelingToolsExp/Private/BakeMeshAttributeVertexTool.cpp
+++ b/Engine/Plugins/Experimental/MeshModelingToolsetExp/Source/MeshModelingToolsExp/Private/BakeMeshAttributeVertexTool.cpp
@@ -295,11 +295,11 @@ void UBakeMeshAttributeVertexTool::Setup()
 	AddToolPropertySource(Settings);
 
 	Settings->WatchProperty(Settings->OutputMode, [this](EBakeVertexOutput) { OpState |= EBakeOpState::Evaluate; UpdateOnModeChange(); });
-	Settings->WatchProperty(Settings->OutputType, [this](EBakeMapType) { OpState |= EBakeOpState::Evaluate; UpdateOnModeChange(); });
-	Settings->WatchProperty(Settings->OutputTypeR, [this](EBakeMapType) { OpState |= EBakeOpState::Evaluate; UpdateOnModeChange(); });
-	Settings->WatchProperty(Settings->OutputTypeG, [this](EBakeMapType) { OpState |= EBakeOpState::Evaluate; UpdateOnModeChange(); });
-	Settings->WatchProperty(Settings->OutputTypeB, [this](EBakeMapType) { OpState |= EBakeOpState::Evaluate; UpdateOnModeChange(); });
-	Settings->WatchProperty(Settings->OutputTypeA, [this](EBakeMapType) { OpState |= EBakeOpState::Evaluate; UpdateOnModeChange(); });
+	Settings->WatchProperty(Settings->OutputType, [this](int32) { OpState |= EBakeOpState::Evaluate; UpdateOnModeChange(); });
+	Settings->WatchProperty(Settings->OutputTypeR, [this](int32) { OpState |= EBakeOpState::Evaluate; UpdateOnModeChange(); });
+	Settings->WatchProperty(Settings->OutputTypeG, [this](int32) { OpState |= EBakeOpState::Evaluate; UpdateOnModeChange(); });
+	Settings->WatchProperty(Settings->OutputTypeB, [this](int32) { OpState |= EBakeOpState::Evaluate; UpdateOnModeChange(); });
+	Settings->WatchProperty(Settings->OutputTypeA, [this](int32) { OpState |= EBakeOpState::Evaluate; UpdateOnModeChange(); });
 	Settings->WatchProperty(Settings->PreviewMode, [this](EBakeVertexChannel) { UpdateVisualization(); });
 	Settings->WatchProperty(Settings->bSplitAtNormalSeams, [this](bool) { bColorTopologyValid = false; OpState |= EBakeOpState::Evaluate; });
 	Settings->WatchProperty(Settings->bSplitAtUVSeams, [this](bool) { bColorTopologyValid = false; OpState |= EBakeOpState::Evaluate; });
@@ -487,7 +487,7 @@ void UBakeMeshAttributeVertexTool::UpdateOnModeChange()
 
 	if (Settings->OutputMode == EBakeVertexOutput::RGBA)
 	{
-		switch (Settings->OutputType)
+		switch (static_cast<EBakeMapType>(Settings->OutputType))
 		{
 			case EBakeMapType::AmbientOcclusion:
 			case EBakeMapType::BentNormal:
@@ -509,11 +509,11 @@ void UBakeMeshAttributeVertexTool::UpdateOnModeChange()
 	}
 	else // Settings->VertexMode == EBakeVertexOutput::PerChannel
 	{
-		EBakeMapType PerChannelTypes[4] = {
-			Settings->OutputTypeR,
-			Settings->OutputTypeG,
-			Settings->OutputTypeB,
-			Settings->OutputTypeA
+		const EBakeMapType PerChannelTypes[4] = {
+			static_cast<EBakeMapType>(Settings->OutputTypeR),
+			static_cast<EBakeMapType>(Settings->OutputTypeG),
+			static_cast<EBakeMapType>(Settings->OutputTypeB),
+			static_cast<EBakeMapType>(Settings->OutputTypeA)
 		};
 		for(int Idx = 0; Idx < 4; ++Idx)
 		{
@@ -630,11 +630,11 @@ void UBakeMeshAttributeVertexTool::UpdateResult()
 
 	FBakeSettings BakeSettings;
 	BakeSettings.OutputMode = Settings->OutputMode;
-	BakeSettings.OutputType = Settings->OutputType;
-	BakeSettings.OutputTypePerChannel[0] = Settings->OutputTypeR;
-	BakeSettings.OutputTypePerChannel[1] = Settings->OutputTypeG;
-	BakeSettings.OutputTypePerChannel[2] = Settings->OutputTypeB;
-	BakeSettings.OutputTypePerChannel[3] = Settings->OutputTypeA;
+	BakeSettings.OutputType = static_cast<EBakeMapType>(Settings->OutputType);
+	BakeSettings.OutputTypePerChannel[0] = static_cast<EBakeMapType>(Settings->OutputTypeR);
+	BakeSettings.OutputTypePerChannel[1] = static_cast<EBakeMapType>(Settings->OutputTypeG);
+	BakeSettings.OutputTypePerChannel[2] = static_cast<EBakeMapType>(Settings->OutputTypeB);
+	BakeSettings.OutputTypePerChannel[3] = static_cast<EBakeMapType>(Settings->OutputTypeA);
 	BakeSettings.bSplitAtNormalSeams = Settings->bSplitAtNormalSeams;
 	BakeSettings.bSplitAtUVSeams = Settings->bSplitAtUVSeams;
 	BakeSettings.bProjectionInWorldSpace = InputMeshSettings->bProjectionInWorldSpace;
diff --git a/Engine/Plugins/Experimental/MeshModelingToolsetExp/Source/MeshModelingToolsExp/Public/BakeMeshAttributeTool.h b/Engine/Plugins/Experimental/MeshModelingToolsetExp/Source/MeshModelingToolsExp/Public/BakeMeshAttributeTool.h
--- a/Engine/Plugins/Experimental/MeshModelingToolsetExp/Source/MeshModelingToolsExp/Public/BakeMeshAttributeTool.h
+++ b/Engine/Plugins/Experimental/MeshModelingToolsetExp/Source/MeshModelingToolsExp/Public/BakeMeshAttributeTool.h
@@ -25,7 +25,7 @@ enum class EBakeMapType
 	None                   = 0,
 
 	/* Normals in tangent space */
-	TangentSpaceNormal     = 1 << 0,
+	TangentSpaceNormal     = 1 << 0 UMETA(DisplayName = "Tangent Normal"),
 	/* Ambient occlusion sampled across the hemisphere */
 	AmbientOcclusion       = 1 << 1,
 	/* Normals skewed towards the least occluded direction */
@@ -35,7 +35,7 @@ enum class EBakeMapType
 	/* Transfer a given texture */
 	Texture                = 1 << 4,
 	/* Interpolated normals in object space */
-	ObjectSpaceNormal      = 1 << 5 UMETA(DisplayName = "Normal"),
+	ObjectSpaceNormal      = 1 << 5 UMETA(DisplayName = "Object Normal"),
 	/* Geometric face normals in object space */
 	FaceNormal             = 1 << 6,
 	/* Positions in object space */
diff --git a/Engine/Plugins/Experimental/MeshModelingToolsetExp/Source/MeshModelingToolsExp/Public/BakeMeshAttributeVertexTool.h b/Engine/Plugins/Experimental/MeshModelingToolsetExp/Source/MeshModelingToolsExp/Public/BakeMeshAttributeVertexTool.h
--- a/Engine/Plugins/Experimental/MeshModelingToolsetExp/Source/MeshModelingToolsExp/Public/BakeMeshAttributeVertexTool.h
+++ b/Engine/Plugins/Experimental/MeshModelingToolsetExp/Source/MeshModelingToolsExp/Public/BakeMeshAttributeVertexTool.h
@@ -63,34 +63,34 @@ public:
 	EBakeVertexOutput OutputMode = EBakeVertexOutput::RGBA;
 
 	/** The bake output type to generate */
-	UPROPERTY(EditAnywhere, Category = BakeOutput, meta=(
+	UPROPERTY(EditAnywhere, Category = BakeOutput, meta=(Bitmask, BitmaskEnum = EBakeMapType,
 		ValidEnumValues="TangentSpaceNormal, AmbientOcclusion, BentNormal, Curvature, Texture, ObjectSpaceNormal, FaceNormal, Position, MaterialID, MultiTexture",
 		EditCondition="OutputMode == EBakeVertexOutput::RGBA", EditConditionHides))
-	EBakeMapType OutputType = EBakeMapType::TangentSpaceNormal;
+	int32 OutputType = static_cast<int32>(EBakeMapType::TangentSpaceNormal);
 
 	/** The bake output type to generate in the Red channel */
-	UPROPERTY(EditAnywhere, Category = BakeOutput, meta=(
+	UPROPERTY(EditAnywhere, Category = BakeOutput, meta=(Bitmask, BitmaskEnum = EBakeMapType,
 		ValidEnumValues="None, AmbientOcclusion, Curvature",
 		EditCondition="OutputMode == EBakeVertexOutput::PerChannel", EditConditionHides))
-	EBakeMapType OutputTypeR = EBakeMapType::None;
+	int32 OutputTypeR = static_cast<int32>(EBakeMapType::None);
 
 	/** The bake output type to generate in the Green channel */
-	UPROPERTY(EditAnywhere, Category = BakeOutput, meta=(
+	UPROPERTY(EditAnywhere, Category = BakeOutput, meta=(Bitmask, BitmaskEnum = EBakeMapType,
 		ValidEnumValues="None, AmbientOcclusion, Curvature",
 		EditCondition="OutputMode == EBakeVertexOutput::PerChannel", EditConditionHides))
-	EBakeMapType OutputTypeG = EBakeMapType::None;
+	int32 OutputTypeG = static_cast<int32>(EBakeMapType::None);
 
 	/** The bake output type to generate in the Blue channel */
-	UPROPERTY(EditAnywhere, Category = BakeOutput, meta=(
+	UPROPERTY(EditAnywhere, Category = BakeOutput, meta=(Bitmask, BitmaskEnum = EBakeMapType,
 		ValidEnumValues="None, AmbientOcclusion, Curvature",
 		EditCondition="OutputMode == EBakeVertexOutput::PerChannel", EditConditionHides))
-	EBakeMapType OutputTypeB = EBakeMapType::None;
+	int32 OutputTypeB = static_cast<int32>(EBakeMapType::None);
 
 	/** The bake output type to generate in the Alpha channel */
-	UPROPERTY(EditAnywhere, Category = BakeOutput, meta=(
+	UPROPERTY(EditAnywhere, Category = BakeOutput, meta=(Bitmask, BitmaskEnum = EBakeMapType,
 		ValidEnumValues="None, AmbientOcclusion, Curvature",
 		EditCondition="OutputMode == EBakeVertexOutput::PerChannel", EditConditionHides))
-	EBakeMapType OutputTypeA = EBakeMapType::None;
+	int32 OutputTypeA = static_cast<int32>(EBakeMapType::None);
 
 	/** The vertex color channel to preview */
 	UPROPERTY(EditAnywhere, Category = BakeOutput, meta = (TransientToolProperty))
diff --git a/Engine/Plugins/Experimental/MeshModelingToolsetExp/Source/ModelingEditorUI/Private/ModelingWidgets/ModelingCustomizationUtil.cpp b/Engine/Plugins/Experimental/MeshModelingToolsetExp/Source/ModelingEditorUI/Private/ModelingWidgets/ModelingCustomizationUtil.cpp
--- a/Engine/Plugins/Experimental/MeshModelingToolsetExp/Source/ModelingEditorUI/Private/ModelingWidgets/ModelingCustomizationUtil.cpp
+++ b/Engine/Plugins/Experimental/MeshModelingToolsetExp/Source/ModelingEditorUI/Private/ModelingWidgets/ModelingCustomizationUtil.cpp
@@ -144,3 +144,27 @@ TSharedRef<SCheckBox> UE::ModelingUI::MakeBoolToggleButton(
 		];
 }
 
+
+void UE::ModelingUI::ProcessChildWidgetsByType(
+	const TSharedRef<SWidget>& RootWidget,
+	const FString& WidgetType,
+	TFunction<bool(TSharedRef<SWidget>&)> ProcessFunc)
+{
+	auto ProcessChildWidgets = [ProcessFunc](const TSharedRef<SWidget>& Widget, const FString& WidgetType, auto& FindRef) -> void
+	{
+		FChildren* Children = Widget->GetChildren();
+		const int32 NumChild = Children ? Children->NumSlot() : 0;
+		for (int32 ChildIdx = 0; ChildIdx < NumChild; ++ChildIdx)
+		{
+			TSharedRef<SWidget> ChildWidget = Children->GetChildAt(ChildIdx);
+			if (ChildWidget->GetTypeAsString().Compare(WidgetType) == 0 && !ProcessFunc(ChildWidget))
+			{
+				break;
+			}
+			FindRef(ChildWidget, WidgetType, FindRef);
+		}
+	};
+	ProcessChildWidgets(RootWidget, WidgetType, ProcessChildWidgets);
+}
+
+
diff --git a/Engine/Plugins/Experimental/MeshModelingToolsetExp/Source/ModelingEditorUI/Public/ModelingWidgets/ModelingCustomizationUtil.h b/Engine/Plugins/Experimental/MeshModelingToolsetExp/Source/ModelingEditorUI/Public/ModelingWidgets/ModelingCustomizationUtil.h
--- a/Engine/Plugins/Experimental/MeshModelingToolsetExp/Source/ModelingEditorUI/Public/ModelingWidgets/ModelingCustomizationUtil.h
+++ b/Engine/Plugins/Experimental/MeshModelingToolsetExp/Source/ModelingEditorUI/Public/ModelingWidgets/ModelingCustomizationUtil.h
@@ -83,6 +83,18 @@ MODELINGEDITORUI_API TSharedRef<SCheckBox> MakeBoolToggleButton(
 	int HorzPadding = 4);
 
 
+/**
+ * Run ProcessFunc over each child widget of WidgetType under RootWidget.
+ * @param RootWidget the root widget to search from.
+ * @param WidgetType the type of widgets to consider
+ * @param ProcessFunc the function to run on each instance of the widget. If return value is false, stop further processing.
+ */
+MODELINGEDITORUI_API void ProcessChildWidgetsByType(
+	const TSharedRef<SWidget>& RootWidget,
+	const FString& WidgetType,
+	TFunction<bool(TSharedRef<SWidget>&)> ProcessFunc);
+
+
 /**
  * Create an error-string widget and add it to the Slot argument, to give the (developer) feedback about slate programming errors
  * @param ErrorString error message
