ModelingTools: Fix bake map types menu widget flickering due to resizing parent SComboWidget.

#rnx

diff --git a/Engine/Plugins/Editor/ModelingToolsEditorMode/Source/ModelingToolsEditorMode/Private/DetailsCustomizations/BakeMeshAttributeToolCustomizations.cpp b/Engine/Plugins/Editor/ModelingToolsEditorMode/Source/ModelingToolsEditorMode/Private/DetailsCustomizations/BakeMeshAttributeToolCustomizations.cpp
--- a/Engine/Plugins/Editor/ModelingToolsEditorMode/Source/ModelingToolsEditorMode/Private/DetailsCustomizations/BakeMeshAttributeToolCustomizations.cpp
+++ b/Engine/Plugins/Editor/ModelingToolsEditorMode/Source/ModelingToolsEditorMode/Private/DetailsCustomizations/BakeMeshAttributeToolCustomizations.cpp
@@ -319,7 +319,17 @@ namespace BakeCustomizationLocals
 			]
 			.ValueContent()
 			[
-				ValueWidget->AsShared()
+				// Fix the size of the SComboButton by setting a MaxWidth HBox.
+				// This addresses the pop-up menu flickering when the parent SComboWidget
+				// is resized while the menu remains visible after selection.
+				SNew(SHorizontalBox)
+				+ SHorizontalBox::Slot()
+				.VAlign(VAlign_Center)
+				.Padding(FMargin(0))
+				.MaxWidth(125.0f)
+				[
+					ValueWidget->AsShared()
+				]
 			];
 	}
 };
