MeshModelingToolset: Rename GenGrps tool to GrpGen

diff --git a/Engine/Plugins/Editor/ModelingToolsEditorMode/Source/ModelingToolsEditorMode/Private/ModelingToolsManagerActions.cpp b/Engine/Plugins/Editor/ModelingToolsEditorMode/Source/ModelingToolsEditorMode/Private/ModelingToolsManagerActions.cpp
--- a/Engine/Plugins/Editor/ModelingToolsEditorMode/Source/ModelingToolsEditorMode/Private/ModelingToolsManagerActions.cpp
+++ b/Engine/Plugins/Editor/ModelingToolsEditorMode/Source/ModelingToolsEditorMode/Private/ModelingToolsManagerActions.cpp
@@ -68,19 +68,63 @@ void FModelingToolsManagerCommands::RegisterCommands()
 	REGISTER_MODELING_TOOL_COMMAND(BeginCubeGridTool, "CubeGr", "Create block out meshes using a repositionable grid");
 	REGISTER_MODELING_TOOL_COMMAND(BeginMeshBooleanTool, "MshBool", "Apply Boolean operations to mesh pairs");
 	REGISTER_MODELING_TOOL_COMMAND(BeginCutMeshWithMeshTool, "MshCut", "Split one mesh into parts using a second mesh");
-	REGISTER_MODELING_TOOL_COMMAND(BeginSubdividePolyTool, "SubDiv", "Subdivide Mesh via PolyGroups or Triangles");
+	REGISTER_MODELING_TOOL_COMMAND(BeginSubdividePolyTool, "SubDiv", "Subdivide mesh via PolyGroups or triangles");
 
 	// TriModel
+	REGISTER_MODELING_TOOL_COMMAND(BeginMeshSelectionTool, "TriSel", "Select and edit mesh triangles");
+	REGISTER_MODELING_TOOL_COMMAND(BeginTriEditTool, "TriEd", "Edit mesh via triangles");
+	REGISTER_MODELING_TOOL_COMMAND(BeginHoleFillTool, "HFill", "Fill holes in mesh");
+	REGISTER_MODELING_TOOL_COMMAND(BeginMirrorTool, "Mirror", "Mirror selected meshes");
+	REGISTER_MODELING_TOOL_COMMAND(BeginPlaneCutTool, "PlnCut", "Cut selected meshes with plane");
+	REGISTER_MODELING_TOOL_COMMAND(BeginPolygonCutTool, "PolyCut", "Cut mesh with extruded polygon");
+	REGISTER_MODELING_TOOL_COMMAND(BeginMeshTrimTool, "Trim", "Trim/cut selected mesh with second mesh");
 
 	// Deform
+	REGISTER_MODELING_TOOL_COMMAND(BeginSculptMeshTool, "VSclpt", "Vertex sculpting");
+	REGISTER_MODELING_TOOL_COMMAND(BeginRemeshSculptMeshTool, "DSclpt", "Dynamic mesh sculpting");
+	REGISTER_MODELING_TOOL_COMMAND(BeginSmoothMeshTool, "Smooth", "Smooth mesh surface");
+	REGISTER_MODELING_TOOL_COMMAND(BeginOffsetMeshTool, "Offset", "Offset mesh surface");
+	REGISTER_MODELING_TOOL_COMMAND(BeginMeshSpaceDeformerTool, "Warp", "Reshape mesh using space deformers");
+	REGISTER_MODELING_TOOL_COMMAND(BeginLatticeDeformerTool, "Lattice", "Deform mesh with 3D lattice/grid");
+	REGISTER_MODELING_TOOL_COMMAND(BeginDisplaceMeshTool, "Displce", "Displace mesh surface with optional subdivision");
 
 	// Transform
+	REGISTER_MODELING_TOOL_COMMAND(BeginTransformMeshesTool, "XForm", "Transform selected meshes");
+	REGISTER_MODELING_TOOL_COMMAND(BeginAlignObjectsTool, "Align", "Align objects");
+	REGISTER_MODELING_TOOL_COMMAND(BeginEditPivotTool, "Pivot", "Edit mesh pivots");
+	REGISTER_MODELING_TOOL_COMMAND(BeginAddPivotActorTool, "PivotAct", "Add actor to act as a pivot for child components");
+	REGISTER_MODELING_TOOL_COMMAND(BeginBakeTransformTool, "BakeRS", "Bake rotation and scale into mesh asset");
+	REGISTER_MODELING_TOOL_COMMAND(BeginTransferMeshTool, "Transfer", "Transfer meshes");
+	REGISTER_MODELING_TOOL_COMMAND(BeginConvertMeshesTool, "Convert", "Convert meshes");
+	REGISTER_MODELING_TOOL_COMMAND(BeginSplitMeshesTool, "Split", "Split meshes");
 
 	// MeshOps
+	REGISTER_MODELING_TOOL_COMMAND(BeginSimplifyMeshTool, "Simplfy", "Simplify mesh");
+	REGISTER_MODELING_TOOL_COMMAND(BeginRemeshMeshTool, "Remesh", "Re-triangulate mesh");
+	REGISTER_MODELING_TOOL_COMMAND(BeginWeldEdgesTool, "Weld", "Weld overlapping mesh edges");
+	REGISTER_MODELING_TOOL_COMMAND(BeginRemoveOccludedTrianglesTool, "Jacket", "Remove hidden triangles from selected meshes");
+	REGISTER_MODELING_TOOL_COMMAND(BeginSelfUnionTool, "Merge", "Self-union selected meshes to resolve self-intersections");
+	REGISTER_MODELING_TOOL_COMMAND(BeginProjectToTargetTool, "Project", "Map/re-mesh onto target mesh (second selection)");
 
 	// VoxOps
+	REGISTER_MODELING_TOOL_COMMAND(BeginVoxelSolidifyTool, "VoxWrap", "Wrap selected meshes using voxels");
+	REGISTER_MODELING_TOOL_COMMAND(BeginVoxelBlendTool, "VoxBlnd", "Blend selected meshes using voxels");
+	REGISTER_MODELING_TOOL_COMMAND(BeginVoxelMorphologyTool, "VoxMrph", "Offset/inset selected meshes using voxels");
+#if WITH_PROXYLOD
+	// The ProxyLOD plugin is currently only available on Windows. Without it, the following tools do not work as expected.
+	REGISTER_MODELING_TOOL_COMMAND(BeginVoxelBooleanTool, "VoxBool", "Boolean selected meshes using voxels");
+	REGISTER_MODELING_TOOL_COMMAND(BeginVoxelMergeTool, "VoxMrg", "Merge selected meshes using voxels");
+#endif	// WITH_PROXYLOD
 
 	// Attributes
+	REGISTER_MODELING_TOOL_COMMAND(BeginMeshInspectorTool, "Inspct", "Inspect mesh attributes");
+	REGISTER_MODELING_TOOL_COMMAND(BeginEditNormalsTool, "Nrmls", "Update mesh normals");
+	REGISTER_MODELING_TOOL_COMMAND(BeginEditTangentsTool, "Tngnts", "Update mesh tangents");
+	REGISTER_MODELING_TOOL_COMMAND(BeginAttributeEditorTool, "AttrEd", "Edit/configure mesh attributes");
+	REGISTER_MODELING_TOOL_COMMAND(BeginPolyGroupsTool, "GrpGen", "Generate new PolyGroups");
+	REGISTER_MODELING_TOOL_COMMAND(BeginMeshGroupPaintTool, "GrpPnt", "Paint new PolyGroups");
+	REGISTER_MODELING_TOOL_COMMAND(BeginMeshAttributePaintTool, "MapPnt", "Paint attribute maps");
+	REGISTER_MODELING_TOOL_COMMAND(BeginEditMeshMaterialsTool, "MatEd", "Assign materials to selected triangles");
 
 	// UVs
 	REGISTER_MODELING_TOOL_COMMAND(BeginGlobalUVGenerateTool, "AutoUV", "Automatically unwrap and pack UVs for mesh");
@@ -96,75 +140,24 @@ void FModelingToolsManagerCommands::RegisterCommands()
 	REGISTER_MODELING_TOOL_COMMAND(BeginBakeMeshAttributeVertexTool, "BakeVtx", "Bake vertex colors for single meshes");
 
 	// Volumes
+	REGISTER_MODELING_TOOL_COMMAND(BeginVolumeToMeshTool, "Vol2Msh", "Convert volume to new mesh asset");
+	REGISTER_MODELING_TOOL_COMMAND(BeginMeshToVolumeTool, "Msh2Vol", "Convert mesh to volume");
+	REGISTER_MODELING_TOOL_COMMAND(BeginBspConversionTool, "BSPConv", "Convert BSP to static mesh asset");
+	REGISTER_MODELING_TOOL_COMMAND(BeginPhysicsInspectorTool, "PInspct", "Inspect physics geometry for selected meshes");
+	REGISTER_MODELING_TOOL_COMMAND(BeginSetCollisionGeometryTool, "Msh2Coll", "Convert selected meshes to Simple Collision Geometry (for last selected)");
+	REGISTER_MODELING_TOOL_COMMAND(BeginExtractCollisionGeometryTool, "Coll2Msh", "Convert Simple Collision Geometry to new mesh asset");
 
 	// LODs
+	REGISTER_MODELING_TOOL_COMMAND(BeginLODManagerTool, "LODMgr", "Static mesh asset LOD manager");
+	REGISTER_MODELING_TOOL_COMMAND(BeginGenerateStaticMeshLODAssetTool, "AutoLOD", "Generate static mesh LOD asset");
 
 	REGISTER_MODELING_TOOL_COMMAND(BeginAddPatchTool, "Patch", "Add Patch");
 	REGISTER_MODELING_TOOL_COMMAND(BeginShapeSprayTool, "Spray", "Shape Spray");
-
-	REGISTER_MODELING_TOOL_COMMAND(BeginSculptMeshTool, "VSclpt", "Vertex Sculpting");
-	REGISTER_MODELING_TOOL_COMMAND(BeginTriEditTool, "TriEd", "Edit Mesh via Triangles");
-	REGISTER_MODELING_TOOL_COMMAND(BeginSmoothMeshTool, "Smooth", "Smooth Mesh surface");
-	REGISTER_MODELING_TOOL_COMMAND(BeginOffsetMeshTool, "Offset", "Offset Mesh surface");
-	REGISTER_MODELING_TOOL_COMMAND(BeginDisplaceMeshTool, "Displce", "Displace Mesh surface with optional subdivision");
-	REGISTER_MODELING_TOOL_COMMAND(BeginMeshSpaceDeformerTool, "Warp", "Reshape Mesh using Space Deformers");
-	REGISTER_MODELING_TOOL_COMMAND(BeginTransformMeshesTool, "XForm", "Transform selected Meshes");
-	REGISTER_MODELING_TOOL_COMMAND(BeginEditPivotTool, "Pivot", "Edit Mesh Pivots");
-	REGISTER_MODELING_TOOL_COMMAND(BeginAddPivotActorTool, "PivotAct", "Add actor to act as a pivot for child component(s)");
-	REGISTER_MODELING_TOOL_COMMAND(BeginBakeTransformTool, "BakeRS", "Bake Scale/Rotation into Mesh Asset");
-	REGISTER_MODELING_TOOL_COMMAND(BeginAlignObjectsTool, "Align", "Align Objects");
-	REGISTER_MODELING_TOOL_COMMAND(BeginTransferMeshTool, "Transfer", "Transfer Mesh");
-	REGISTER_MODELING_TOOL_COMMAND(BeginConvertMeshesTool, "Convert", "Convert Meshes");
-	REGISTER_MODELING_TOOL_COMMAND(BeginSplitMeshesTool, "Split", "Split Meshes");
-
-	REGISTER_MODELING_TOOL_COMMAND(BeginRemeshSculptMeshTool, "DSclpt", "Dynamic Mesh Sculpting");
-	REGISTER_MODELING_TOOL_COMMAND(BeginRemeshMeshTool, "Remesh", "Retriangulate Mesh");
-	REGISTER_MODELING_TOOL_COMMAND(BeginProjectToTargetTool, "Project", "Map/Remesh onto Target Mesh (second selection)");
-	REGISTER_MODELING_TOOL_COMMAND(BeginSimplifyMeshTool, "Simplfy", "Simplify Mesh");
-	REGISTER_MODELING_TOOL_COMMAND(BeginEditNormalsTool, "Nrmls", "Modify/Update Mesh Normals");
-	REGISTER_MODELING_TOOL_COMMAND(BeginEditTangentsTool, "Tngnts", "Update Mesh Tangents");
-	REGISTER_MODELING_TOOL_COMMAND(BeginRemoveOccludedTrianglesTool, "Jacket", "Remove Hidden Triangles from selected Meshes");
-#if WITH_PROXYLOD
-	// The ProxyLOD plugin is currently only available on Windows. Without it, the following tools do not work as expected.
-	REGISTER_MODELING_TOOL_COMMAND(BeginVoxelMergeTool, "VoxMrg", "Merge Selected Meshes (Voxel-Based)");
-	REGISTER_MODELING_TOOL_COMMAND(BeginVoxelBooleanTool, "VoxBool", "Boolean Selected Meshes (Voxel-Based)");
-#endif	// WITH_PROXYLOD
-	REGISTER_MODELING_TOOL_COMMAND(BeginVoxelSolidifyTool, "VoxWrap", "Wrap Selected Meshes (Voxel-Based)");
-	REGISTER_MODELING_TOOL_COMMAND(BeginVoxelBlendTool, "VoxBlnd", "Blend Selected Meshes (Voxel-Based)");
-	REGISTER_MODELING_TOOL_COMMAND(BeginVoxelMorphologyTool, "VoxMrph", "Offset/Inset Selected Meshes (Voxel-Based)");
-	REGISTER_MODELING_TOOL_COMMAND(BeginSelfUnionTool, "Merge", "Self-Union Selected Meshes to resolve Self-Intersections");
-	REGISTER_MODELING_TOOL_COMMAND(BeginMeshTrimTool, "Trim", "Trim/Cut selected mesh with second mesh");
-	REGISTER_MODELING_TOOL_COMMAND(BeginBspConversionTool, "BSPConv", "Convert BSP to StaticMesh Asset");
-	REGISTER_MODELING_TOOL_COMMAND(BeginMeshToVolumeTool, "Msh2Vol", "Convert Mesh to Volume");
-	REGISTER_MODELING_TOOL_COMMAND(BeginVolumeToMeshTool, "Vol2Msh", "Convert Volume to new Mesh Asset");
-	REGISTER_MODELING_TOOL_COMMAND(BeginPlaneCutTool, "PlnCut", "Cut Selected Meshes with Plane");
-	REGISTER_MODELING_TOOL_COMMAND(BeginMirrorTool, "Mirror", "Mirror Selected Meshes");
-	REGISTER_MODELING_TOOL_COMMAND(BeginHoleFillTool, "HFill", "Fill Holes in Mesh");
-	REGISTER_MODELING_TOOL_COMMAND(BeginLatticeDeformerTool, "Lattice", "Deform Mesh with 3D Lattice/Grid");
-	REGISTER_MODELING_TOOL_COMMAND(BeginPolygonCutTool, "PolyCut", "Cut Mesh with Extruded Polygon");
-	
-	REGISTER_MODELING_TOOL_COMMAND(BeginMeshSelectionTool, "TriSel", "Select and Edit Mesh Triangles");
-
-	REGISTER_MODELING_TOOL_COMMAND(BeginPhysicsInspectorTool, "PInspct", "Inspect Physics Geometry for selected Meshes");
-	REGISTER_MODELING_TOOL_COMMAND(BeginSetCollisionGeometryTool, "Msh2Coll", "Convert Selected Meshes to Simple Collision Geometry (for last Selected)");
 	REGISTER_MODELING_TOOL_COMMAND(BeginEditCollisionGeometryTool, "EditPhys", "Edit Simple Collision Geometry for selected Mesh");
-	REGISTER_MODELING_TOOL_COMMAND(BeginExtractCollisionGeometryTool, "Coll2Msh", "Convert Simple Collision Geometry to new Mesh Asset");
-
-	REGISTER_MODELING_TOOL_COMMAND(BeginMeshInspectorTool, "Inspct", "Inspect Mesh Attributes");
-	REGISTER_MODELING_TOOL_COMMAND(BeginWeldEdgesTool, "Weld", "Weld Overlapping Mesh Edges");
-	REGISTER_MODELING_TOOL_COMMAND(BeginPolyGroupsTool, "GenGrps", "Generate new PolyGroups");
-	REGISTER_MODELING_TOOL_COMMAND(BeginEditMeshMaterialsTool, "MatEd", "Assign Materials to Selected Triangles");
-	REGISTER_MODELING_TOOL_COMMAND(BeginMeshAttributePaintTool, "MapPnt", "Paint Attribute Maps");
-	REGISTER_MODELING_TOOL_COMMAND(BeginAttributeEditorTool, "AttrEd", "Edit/Configure Mesh Attributes");
 
 	// why are these ::Button ?
 	UI_COMMAND(BeginSkinWeightsPaintTool, "SkinWts", "Start the Paint Skin Weights Tool", EUserInterfaceActionType::Button, FInputChord());
 	UI_COMMAND(BeginSkinWeightsBindingTool, "SkinBind", "Start the Skin Weights Binding Tool", EUserInterfaceActionType::Button, FInputChord());
-	
-	REGISTER_MODELING_TOOL_COMMAND(BeginLODManagerTool, "LODMgr", "Static Mesh Asset LOD Manager");
-	REGISTER_MODELING_TOOL_COMMAND(BeginGenerateStaticMeshLODAssetTool, "AutoLOD", "Generate Static Mesh LOD Asset");
-	REGISTER_MODELING_TOOL_COMMAND(BeginMeshGroupPaintTool, "GrpPnt", "Paint New Mesh Polygroups");
-
 
 	REGISTER_MODELING_TOOL_COMMAND(BeginPolyModelTool_FaceSelect, "Faces", "PolyGroup Face Selection Tool");
 	REGISTER_MODELING_TOOL_COMMAND(BeginPolyModelTool_EdgeSelect, "Edges", "PolyGroup Edge Selection Tool");
diff --git a/Engine/Plugins/Runtime/MeshModelingToolset/Source/MeshModelingTools/Private/EditMeshPolygonsTool.cpp b/Engine/Plugins/Runtime/MeshModelingToolset/Source/MeshModelingTools/Private/EditMeshPolygonsTool.cpp
--- a/Engine/Plugins/Runtime/MeshModelingToolset/Source/MeshModelingTools/Private/EditMeshPolygonsTool.cpp
+++ b/Engine/Plugins/Runtime/MeshModelingToolset/Source/MeshModelingTools/Private/EditMeshPolygonsTool.cpp
@@ -420,7 +420,7 @@ void UEditMeshPolygonsTool::Setup()
 	{
 		GetToolManager()->DisplayMessage(
 			LOCTEXT("NoGroupsWarning",
-			        "This object has only a single Polygroup. Use the GenGrps, GrpPnt or TriSel (Create Polygroup) tools to modify Polygroups."),
+			        "This object has only a single Polygroup. Use the GrpGen, GrpPnt or TriSel (Create Polygroup) tools to modify PolyGroups."),
 			EToolMessageLevel::UserWarning);
 	}
 
