ModelingTools: Updated VSculpt icons.

#rnx

diff --git a/Engine/Plugins/Editor/ModelingToolsEditorMode/Source/ModelingToolsEditorMode/Private/ModelingToolsEditorModeStyle.cpp b/Engine/Plugins/Editor/ModelingToolsEditorMode/Source/ModelingToolsEditorMode/Private/ModelingToolsEditorModeStyle.cpp
--- a/Engine/Plugins/Editor/ModelingToolsEditorMode/Source/ModelingToolsEditorMode/Private/ModelingToolsEditorModeStyle.cpp
+++ b/Engine/Plugins/Editor/ModelingToolsEditorMode/Source/ModelingToolsEditorMode/Private/ModelingToolsEditorModeStyle.cpp
@@ -313,36 +313,36 @@ void FModelingToolsEditorModeStyle::Initialize()
 		// Icons for brush falloffs in sculpt/etc tools
 		//
 
-		StyleSet->Set("BrushFalloffIcons.Smooth", new IMAGE_PLUGIN_BRUSH("Icons/BrushIcons/Falloff_Smooth", Icon120));
-		StyleSet->Set("BrushFalloffIcons.Linear", new IMAGE_PLUGIN_BRUSH("Icons/BrushIcons/Falloff_Linear", Icon120));
-		StyleSet->Set("BrushFalloffIcons.Inverse", new IMAGE_PLUGIN_BRUSH("Icons/BrushIcons/Falloff_Inverse", Icon120));
-		StyleSet->Set("BrushFalloffIcons.Round", new IMAGE_PLUGIN_BRUSH("Icons/BrushIcons/Falloff_Round", Icon120));
-		StyleSet->Set("BrushFalloffIcons.BoxSmooth", new IMAGE_PLUGIN_BRUSH("Icons/BrushIcons/Falloff_BoxSmooth", Icon120));
-		StyleSet->Set("BrushFalloffIcons.BoxLinear", new IMAGE_PLUGIN_BRUSH("Icons/BrushIcons/Falloff_BoxLinear", Icon120));
-		StyleSet->Set("BrushFalloffIcons.BoxInverse", new IMAGE_PLUGIN_BRUSH("Icons/BrushIcons/Falloff_BoxInverse", Icon120));
-		StyleSet->Set("BrushFalloffIcons.BoxRound", new IMAGE_PLUGIN_BRUSH("Icons/BrushIcons/Falloff_BoxRound", Icon120));
+		StyleSet->Set("BrushFalloffIcons.Smooth", new IMAGE_BRUSH_SVG("Icons/BrushIcons/Falloff_Smooth", Icon120));
+		StyleSet->Set("BrushFalloffIcons.Linear", new IMAGE_BRUSH_SVG("Icons/BrushIcons/Falloff_Linear", Icon120));
+		StyleSet->Set("BrushFalloffIcons.Inverse", new IMAGE_BRUSH_SVG("Icons/BrushIcons/Falloff_Inverse", Icon120));
+		StyleSet->Set("BrushFalloffIcons.Round", new IMAGE_BRUSH_SVG("Icons/BrushIcons/Falloff_Round", Icon120));
+		StyleSet->Set("BrushFalloffIcons.BoxSmooth", new IMAGE_BRUSH_SVG("Icons/BrushIcons/Falloff_BoxSmooth", Icon120));
+		StyleSet->Set("BrushFalloffIcons.BoxLinear", new IMAGE_BRUSH_SVG("Icons/BrushIcons/Falloff_BoxLinear", Icon120));
+		StyleSet->Set("BrushFalloffIcons.BoxInverse", new IMAGE_BRUSH_SVG("Icons/BrushIcons/Falloff_BoxInverse", Icon120));
+		StyleSet->Set("BrushFalloffIcons.BoxRound", new IMAGE_BRUSH_SVG("Icons/BrushIcons/Falloff_BoxRound", Icon120));
 
 
 		//
 		// Icons for brushes in sculpt/etc tools
 		//
 
-		StyleSet->Set("BrushTypeIcons.Smooth", new IMAGE_PLUGIN_BRUSH("Icons/BrushIcons/Brush_Smooth", Icon120));
-		StyleSet->Set("BrushTypeIcons.SmoothFill", new IMAGE_PLUGIN_BRUSH("Icons/BrushIcons/Brush_SmoothFill", Icon120));
-		StyleSet->Set("BrushTypeIcons.Move", new IMAGE_PLUGIN_BRUSH("Icons/BrushIcons/Brush_Move", Icon120));
-		StyleSet->Set("BrushTypeIcons.SculptN", new IMAGE_PLUGIN_BRUSH("Icons/BrushIcons/Brush_SculptN", Icon120));
-		StyleSet->Set("BrushTypeIcons.SculptV", new IMAGE_PLUGIN_BRUSH("Icons/BrushIcons/Brush_SculptV", Icon120));
-		StyleSet->Set("BrushTypeIcons.SculptMx", new IMAGE_PLUGIN_BRUSH("Icons/BrushIcons/Brush_SculptMx", Icon120));
-		StyleSet->Set("BrushTypeIcons.Inflate", new IMAGE_PLUGIN_BRUSH("Icons/BrushIcons/Brush_Inflate", Icon120));
-		StyleSet->Set("BrushTypeIcons.Pinch", new IMAGE_PLUGIN_BRUSH("Icons/BrushIcons/Brush_Pinch", Icon120));
-		StyleSet->Set("BrushTypeIcons.Flatten", new IMAGE_PLUGIN_BRUSH("Icons/BrushIcons/Brush_Flatten", Icon120));
-		StyleSet->Set("BrushTypeIcons.PlaneN", new IMAGE_PLUGIN_BRUSH("Icons/BrushIcons/Brush_PlaneN", Icon120));
-		StyleSet->Set("BrushTypeIcons.PlaneV", new IMAGE_PLUGIN_BRUSH("Icons/BrushIcons/Brush_PlaneV", Icon120));
-		StyleSet->Set("BrushTypeIcons.PlaneW", new IMAGE_PLUGIN_BRUSH("Icons/BrushIcons/Brush_PlaneW", Icon120));
-		StyleSet->Set("BrushTypeIcons.Scale", new IMAGE_PLUGIN_BRUSH("Icons/BrushIcons/Brush_Scale", Icon120));
-		StyleSet->Set("BrushTypeIcons.Grab", new IMAGE_PLUGIN_BRUSH("Icons/BrushIcons/Brush_Grab", Icon120));
-		StyleSet->Set("BrushTypeIcons.GrabSharp", new IMAGE_PLUGIN_BRUSH("Icons/BrushIcons/Brush_GrabSharp", Icon120));
-		StyleSet->Set("BrushTypeIcons.Twist", new IMAGE_PLUGIN_BRUSH("Icons/BrushIcons/Brush_Twist", Icon120));
+		StyleSet->Set("BrushTypeIcons.Smooth", new IMAGE_BRUSH_SVG("Icons/BrushIcons/Brush_Smooth", Icon120));
+		StyleSet->Set("BrushTypeIcons.SmoothFill", new IMAGE_BRUSH_SVG("Icons/BrushIcons/Brush_SmoothFill", Icon120));
+		StyleSet->Set("BrushTypeIcons.Move", new IMAGE_BRUSH_SVG("Icons/BrushIcons/Brush_Move", Icon120));
+		StyleSet->Set("BrushTypeIcons.SculptN", new IMAGE_BRUSH_SVG("Icons/BrushIcons/Brush_SculptN", Icon120));
+		StyleSet->Set("BrushTypeIcons.SculptV", new IMAGE_BRUSH_SVG("Icons/BrushIcons/Brush_SculptV", Icon120));
+		StyleSet->Set("BrushTypeIcons.SculptMx", new IMAGE_BRUSH_SVG("Icons/BrushIcons/Brush_SculptMx", Icon120));
+		StyleSet->Set("BrushTypeIcons.Inflate", new IMAGE_BRUSH_SVG("Icons/BrushIcons/Brush_Inflate", Icon120));
+		StyleSet->Set("BrushTypeIcons.Pinch", new IMAGE_BRUSH_SVG("Icons/BrushIcons/Brush_Pinch", Icon120));
+		StyleSet->Set("BrushTypeIcons.Flatten", new IMAGE_BRUSH_SVG("Icons/BrushIcons/Brush_Flatten", Icon120));
+		StyleSet->Set("BrushTypeIcons.PlaneN", new IMAGE_BRUSH_SVG("Icons/BrushIcons/Brush_PlaneN", Icon120));
+		StyleSet->Set("BrushTypeIcons.PlaneV", new IMAGE_BRUSH_SVG("Icons/BrushIcons/Brush_PlaneV", Icon120));
+		StyleSet->Set("BrushTypeIcons.PlaneW", new IMAGE_BRUSH_SVG("Icons/BrushIcons/Brush_PlaneW", Icon120));
+		StyleSet->Set("BrushTypeIcons.Scale", new IMAGE_BRUSH_SVG("Icons/BrushIcons/Brush_Scale", Icon120));
+		StyleSet->Set("BrushTypeIcons.Grab", new IMAGE_BRUSH_SVG("Icons/BrushIcons/Brush_Grab", Icon120));
+		StyleSet->Set("BrushTypeIcons.GrabSharp", new IMAGE_BRUSH_SVG("Icons/BrushIcons/Brush_GrabSharp", Icon120));
+		StyleSet->Set("BrushTypeIcons.Twist", new IMAGE_BRUSH_SVG("Icons/BrushIcons/Brush_Twist", Icon120));
 
 	}
 
