localization macros duplicate-key fixes
#rnx

diff --git a/Engine/Plugins/Editor/ModelingToolsEditorMode/Source/ModelingToolsEditorMode/Private/DetailsCustomizations/MeshVertexSculptToolCustomizations.cpp b/Engine/Plugins/Editor/ModelingToolsEditorMode/Source/ModelingToolsEditorMode/Private/DetailsCustomizations/MeshVertexSculptToolCustomizations.cpp
--- a/Engine/Plugins/Editor/ModelingToolsEditorMode/Source/ModelingToolsEditorMode/Private/DetailsCustomizations/MeshVertexSculptToolCustomizations.cpp
+++ b/Engine/Plugins/Editor/ModelingToolsEditorMode/Source/ModelingToolsEditorMode/Private/DetailsCustomizations/MeshVertexSculptToolCustomizations.cpp
@@ -343,7 +343,7 @@ TSharedPtr<SWidget> FVertexBrushSculptPropertiesDetails::MakeFreezeTargetWidget(
 				SNew(STextBlock)
 				.Justification(ETextJustify::Center)
 				.TextStyle(FAppStyle::Get(), "DetailsView.CategoryTextStyle")
-				.Text_Lambda( [this] { return IsFreezeTargetEnabled() ? LOCTEXT("FreezeTarget", "UnFreeze Target") : LOCTEXT("FreezeTarget", "Freeze Target"); } )
+				.Text_Lambda( [this] { return IsFreezeTargetEnabled() ? LOCTEXT("UnFreezeTarget", "UnFreeze Target") : LOCTEXT("FreezeTarget", "Freeze Target"); } )
 			]
 		];
 }
diff --git a/Engine/Plugins/Experimental/GeometryScripting/Source/GeometryScriptingCore/Private/MeshAssetFunctions.cpp b/Engine/Plugins/Experimental/GeometryScripting/Source/GeometryScriptingCore/Private/MeshAssetFunctions.cpp
--- a/Engine/Plugins/Experimental/GeometryScripting/Source/GeometryScriptingCore/Private/MeshAssetFunctions.cpp
+++ b/Engine/Plugins/Experimental/GeometryScripting/Source/GeometryScriptingCore/Private/MeshAssetFunctions.cpp
@@ -194,7 +194,7 @@ UDynamicMesh*  UGeometryScriptLibrary_StaticMeshFunctions::CopyMeshToStaticMesh(
 	if (!ensure(MeshDescription != nullptr))
 	{
 		UE::Geometry::AppendError(Debug, EGeometryScriptErrorType::InvalidInputs,
-			FText::Format(LOCTEXT("CopyMeshToAsset_EngineAsset", "CopyMeshToAsset: MeshDescription for LOD {0} is null?"), FText::AsNumber(UseLODIndex)));
+			FText::Format(LOCTEXT("CopyMeshToAsset_NullMeshDescription", "CopyMeshToAsset: MeshDescription for LOD {0} is null?"), FText::AsNumber(UseLODIndex)));
 		return FromDynamicMesh;
 	}
 
diff --git a/Engine/Plugins/Experimental/MeshModelingToolsetExp/Source/MeshModelingToolsExp/Private/MeshGroupPaintTool.cpp b/Engine/Plugins/Experimental/MeshModelingToolsetExp/Source/MeshModelingToolsExp/Private/MeshGroupPaintTool.cpp
--- a/Engine/Plugins/Experimental/MeshModelingToolsetExp/Source/MeshModelingToolsExp/Private/MeshGroupPaintTool.cpp
+++ b/Engine/Plugins/Experimental/MeshModelingToolsetExp/Source/MeshModelingToolsExp/Private/MeshGroupPaintTool.cpp
@@ -1229,7 +1229,7 @@ void UMeshGroupPaintTool::ToggleFrozenGroup(int32 FreezeGroupID)
 			TempROIBuffer.Add(tid);
 		}
 	}
-	EmitFrozenGroupsChange(InitialFrozenGroups, FrozenGroups, LOCTEXT("ToggleFrozenGroup", "Toggle Frozen Group"));
+	EmitFrozenGroupsChange(InitialFrozenGroups, FrozenGroups, LOCTEXT("ToggleFrozenGroupChange", "Toggle Frozen Group"));
 	DynamicMeshComponent->FastNotifyTriangleVerticesUpdated(TempROIBuffer, EMeshRenderAttributeFlags::VertexColors);
 	GetToolManager()->PostInvalidation();
 }
diff --git a/Engine/Plugins/Runtime/MeshModelingToolset/Source/MeshModelingTools/Private/ToolActivities/PolyEditBevelEdgeActivity.cpp b/Engine/Plugins/Runtime/MeshModelingToolset/Source/MeshModelingTools/Private/ToolActivities/PolyEditBevelEdgeActivity.cpp
--- a/Engine/Plugins/Runtime/MeshModelingToolset/Source/MeshModelingTools/Private/ToolActivities/PolyEditBevelEdgeActivity.cpp
+++ b/Engine/Plugins/Runtime/MeshModelingToolset/Source/MeshModelingTools/Private/ToolActivities/PolyEditBevelEdgeActivity.cpp
@@ -252,7 +252,7 @@ void UPolyEditBevelEdgeActivity::ApplyBevel()
 
 	if (ResultMesh->TriangleCount() == 0)
 	{
-		ParentTool->GetToolManager()->DisplayMessage( LOCTEXT("OnBevelFailedMessage", "Bevel created empty mesh, ignoring"), EToolMessageLevel::UserWarning);
+		ParentTool->GetToolManager()->DisplayMessage( LOCTEXT("OnBevelEmptyMeshMessage", "Bevel created empty mesh, ignoring"), EToolMessageLevel::UserWarning);
 		// Reset the preview
 		ActivityContext->Preview->PreviewMesh->UpdatePreview(ActivityContext->CurrentMesh.Get());
 		return;
