ModelingMode: Adjust quick settings in restrictive mode

diff --git a/Engine/Plugins/Editor/ModelingToolsEditorMode/Source/ModelingToolsEditorMode/Private/ModelingModeAssetUtils.cpp b/Engine/Plugins/Editor/ModelingToolsEditorMode/Source/ModelingToolsEditorMode/Private/ModelingModeAssetUtils.cpp
--- a/Engine/Plugins/Editor/ModelingToolsEditorMode/Source/ModelingToolsEditorMode/Private/ModelingModeAssetUtils.cpp
+++ b/Engine/Plugins/Editor/ModelingToolsEditorMode/Source/ModelingToolsEditorMode/Private/ModelingModeAssetUtils.cpp
@@ -83,7 +83,7 @@ FString InteractiveSelectAssetPath(const FString& DefaultAssetName, const FText&
 FString UE::Modeling::GetNewAssetPathName(const FString& BaseNameIn, const UWorld* TargetWorld, FString SuggestedFolder)
 {
 	const UModelingToolsEditorModeSettings* Settings = GetDefault<UModelingToolsEditorModeSettings>();
-	EModelingModeAssetGenerationBehavior AutoGenMode = Settings->AssetGenerationMode;
+	EModelingModeAssetGenerationBehavior AutoGenMode = Settings->GetAssetGenerationMode();
 
 	FString PackageFolderPath = TEXT("/Game/");
 
@@ -96,7 +96,7 @@ FString UE::Modeling::GetNewAssetPathName(const FString& BaseNameIn, const UWorl
 	{
 		// Figure out the appropriate root path to use. If the level is in a plugin, this should
 		// be the plugin top-level path and not /Game.
-		if (Settings->AssetGenerationLocation == EModelingModeAssetGenerationLocation::AutoGeneratedWorldRelativeAssetPath)
+		if (Settings->GetAssetGenerationLocation() == EModelingModeAssetGenerationLocation::AutoGeneratedWorldRelativeAssetPath)
 		{
 			PackageFolderPath = UE::Local::GetWorldRelativeAssetRootPath(TargetWorld);
 		}
@@ -116,22 +116,41 @@ FString UE::Modeling::GetNewAssetPathName(const FString& BaseNameIn, const UWorl
 			PackageFolderPath = FPaths::Combine(PackageFolderPath, Settings->AutoGeneratedAssetPath);
 		}
 
-		// append username-specific subfolder
-		if (Settings->bUsePerUserAutogenSubfolder)
+		if (Settings->InRestrictiveMode())
 		{
-			FString UsernameString = Settings->AutogenSubfolderUserNameOverride.TrimStartAndEnd();
-			if (UsernameString.IsEmpty())
+			const FString AutogenSubfolder = Settings->GetRestrictiveModeAutogenSubfolder().TrimStartAndEnd();
+			if (!AutogenSubfolder.IsEmpty())
 			{
-				UsernameString = FPlatformProcess::UserName();
+				PackageFolderPath = FPaths::Combine(PackageFolderPath, AutogenSubfolder);
 			}
-			if (!UsernameString.IsEmpty())
+			else
 			{
-				PackageFolderPath = FPaths::Combine(PackageFolderPath, UsernameString);
+				FString UsernameString = FPlatformProcess::UserName();
+				if (!UsernameString.IsEmpty())
+				{
+					PackageFolderPath = FPaths::Combine(PackageFolderPath, UsernameString);
+				}
+			}
+		}
+		else
+		{
+			// append username-specific subfolder
+			if (Settings->bUsePerUserAutogenSubfolder)
+			{
+				FString UsernameString = Settings->AutogenSubfolderUserNameOverride.TrimStartAndEnd();
+				if (UsernameString.IsEmpty())
+				{
+					UsernameString = FPlatformProcess::UserName();
+				}
+				if (!UsernameString.IsEmpty())
+				{
+					PackageFolderPath = FPaths::Combine(PackageFolderPath, UsernameString);
+				}
 			}
 		}
 
 		// if we want to use the currently-visible asset browser path, try to find one (this can fail if no asset browser is visible/etc)
-		if (Settings->AssetGenerationLocation == EModelingModeAssetGenerationLocation::CurrentAssetBrowserPathIfAvailable)
+		if (Settings->GetAssetGenerationLocation() == EModelingModeAssetGenerationLocation::CurrentAssetBrowserPathIfAvailable)
 		{
 			FString CurrentAssetPath = UE::Local::GetActiveAssetFolderPath();
 			if (CurrentAssetPath.IsEmpty() == false)
@@ -206,7 +225,7 @@ bool UE::Modeling::AutoSaveAsset(UObject* Asset)
 void UE::Modeling::OnNewAssetCreated(UObject* Asset)
 {
 	const UModelingToolsEditorModeSettings* Settings = GetDefault<UModelingToolsEditorModeSettings>();
-	EModelingModeAssetGenerationBehavior AutoGenMode = Settings->AssetGenerationMode;
+	EModelingModeAssetGenerationBehavior AutoGenMode = Settings->GetAssetGenerationMode();
 
 	// Auto-save the new asset if the user has requested this behavior, or if an interactive
 	// save/path dialog was shown earlier in the process (InteractivePromptToSave mode)
diff --git a/Engine/Plugins/Editor/ModelingToolsEditorMode/Source/ModelingToolsEditorMode/Private/ModelingToolsEditorModeToolkit.cpp b/Engine/Plugins/Editor/ModelingToolsEditorMode/Source/ModelingToolsEditorMode/Private/ModelingToolsEditorModeToolkit.cpp
--- a/Engine/Plugins/Editor/ModelingToolsEditorMode/Source/ModelingToolsEditorMode/Private/ModelingToolsEditorModeToolkit.cpp
+++ b/Engine/Plugins/Editor/ModelingToolsEditorMode/Source/ModelingToolsEditorMode/Private/ModelingToolsEditorModeToolkit.cpp
@@ -305,9 +305,11 @@ TSharedPtr<SWidget> FModelingToolsEditorModeToolkit::MakeAssetConfigPanel()
 
 	});
 
-	AssetLODModeLabel = SNew(STextBlock).Text(LOCTEXT("ActiveLODLabel", "Editing LOD"));
+	AssetLODModeLabel = SNew(STextBlock)
+	.Text(LOCTEXT("ActiveLODLabel", "Editing LOD"))
+	.ToolTipText(LOCTEXT("ActiveLODLabelToolTip", "Select the LOD to be used when editing an existing mesh."));
 
-	TSharedPtr<SVerticalBox> Content = SNew(SVerticalBox)
+	const TSharedPtr<SVerticalBox> Content = SNew(SVerticalBox)
 	+ SVerticalBox::Slot().HAlign(HAlign_Fill)
 		[
 			SNew(SHorizontalBox)
@@ -315,28 +317,50 @@ TSharedPtr<SWidget> FModelingToolsEditorModeToolkit::MakeAssetConfigPanel()
 				[ AssetLODModeLabel->AsShared() ]
 			+ SHorizontalBox::Slot().Padding(0).FillWidth(4.0f)
 				[ AssetLODMode->AsShared() ]
-		]
+		];
 
-	+ SVerticalBox::Slot().HAlign(HAlign_Fill).Padding(0, 3, 0, 0)
-		[
+	if (Settings->InRestrictiveMode())
+	{
+		AutogenSubfolder = SNew(SEditableTextBox)
+		.Text(this, &FModelingToolsEditorModeToolkit::GetRestrictiveModeAutogenSubfolderText)
+		.OnTextCommitted(this, &FModelingToolsEditorModeToolkit::OnRestrictiveModeAutogenSubfolderTextCommitted);
+		
+		Content->AddSlot()[
 			SNew(SHorizontalBox)
 			+ SHorizontalBox::Slot().Padding(0, 2, 2, 2).HAlign(HAlign_Right).VAlign(VAlign_Center).AutoWidth()
-				[
-					SNew(STextBlock)
-					.Text(LOCTEXT("AssetLocationLabel", "New Asset Location"))
-				]
+			[
+				SNew(STextBlock)
+				.Text(LOCTEXT("AutogenSubfolder", "Autogen Subfolder"))
+				.ToolTipText(LOCTEXT("AutogenSubfolderToolTip", "Override for the name of the subfolder in which new assets get stored in. "
+					"If nothing is set, the username is used as the subfolder name."))
+			]
 			+ SHorizontalBox::Slot().HAlign(HAlign_Fill).Padding(0).FillWidth(1.0f)
-				[
-					AssetLocationMode->AsShared()
-				]
+			[
+				AutogenSubfolder->AsShared()
+			]
+		];
+	}
+	else
+	{
+		Content->AddSlot()[
+			SNew(SHorizontalBox)
+			+ SHorizontalBox::Slot().Padding(0, 2, 2, 2).HAlign(HAlign_Right).VAlign(VAlign_Center).AutoWidth()
+			[
+				SNew(STextBlock)
+				.Text(LOCTEXT("AssetLocationLabel", "New Asset Location"))
+			]
+			+ SHorizontalBox::Slot().HAlign(HAlign_Fill).Padding(0).FillWidth(1.0f)
+			[
+				AssetLocationMode->AsShared()
+			]
 			+ SHorizontalBox::Slot().HAlign(HAlign_Right).Padding(0, 0, 0, 0).AutoWidth()
-				[
-					SNew(SSimpleButton)
-					.OnClicked_Lambda( [this]() { OnShowAssetSettings(); return FReply::Handled(); } )
-					.Icon(FAppStyle::Get().GetBrush("Icons.Settings"))
-				]
+			[
+				SNew(SSimpleButton)
+				.OnClicked_Lambda( [this]() { OnShowAssetSettings(); return FReply::Handled(); } )
+				.Icon(FAppStyle::Get().GetBrush("Icons.Settings"))
+			]
 		];
-
+	}
 
 	TSharedPtr<SExpandableArea> AssetConfigPanel = SNew(SExpandableArea)
 		.HeaderPadding(FMargin(2.0f))
@@ -997,46 +1021,44 @@ void FModelingToolsEditorModeToolkit::UpdateShowWarnings()
 void FModelingToolsEditorModeToolkit::UpdateAssetLocationMode(TSharedPtr<FString> NewString)
 {
 	UModelingToolsEditorModeSettings* Settings = GetMutableDefault<UModelingToolsEditorModeSettings>();
+	EModelingModeAssetGenerationLocation AssetGenerationLocation = EModelingModeAssetGenerationLocation::AutoGeneratedWorldRelativeAssetPath;
 	if (NewString == AssetLocationModes[0])
 	{
-		Settings->AssetGenerationLocation = EModelingModeAssetGenerationLocation::AutoGeneratedWorldRelativeAssetPath;
+		AssetGenerationLocation = EModelingModeAssetGenerationLocation::AutoGeneratedWorldRelativeAssetPath;
 	}
-	if (NewString == AssetLocationModes[1])
+	else if (NewString == AssetLocationModes[1])
 	{
-		Settings->AssetGenerationLocation = EModelingModeAssetGenerationLocation::AutoGeneratedGlobalAssetPath;
+		AssetGenerationLocation = EModelingModeAssetGenerationLocation::AutoGeneratedGlobalAssetPath;
 	}
-	else if ( NewString == AssetLocationModes[2])
+	else if (NewString == AssetLocationModes[2])
 	{
-		Settings->AssetGenerationLocation = EModelingModeAssetGenerationLocation::CurrentAssetBrowserPathIfAvailable;
-	}
-	else
-	{
-		Settings->AssetGenerationLocation = EModelingModeAssetGenerationLocation::AutoGeneratedWorldRelativeAssetPath;
+		AssetGenerationLocation = EModelingModeAssetGenerationLocation::CurrentAssetBrowserPathIfAvailable;
 	}
 
+	Settings->SetAssetGenerationLocation(AssetGenerationLocation);
+	
 	Settings->SaveConfig();
 }
 
 void FModelingToolsEditorModeToolkit::UpdateAssetSaveMode(TSharedPtr<FString> NewString)
 {
 	UModelingToolsEditorModeSettings* Settings = GetMutableDefault<UModelingToolsEditorModeSettings>();
+	EModelingModeAssetGenerationBehavior AssetGenerationBehavior = EModelingModeAssetGenerationBehavior::AutoGenerateButDoNotAutosave;
 	if (NewString == AssetSaveModes[0])
 	{
-		Settings->AssetGenerationMode = EModelingModeAssetGenerationBehavior::AutoGenerateAndAutosave;
+		AssetGenerationBehavior = EModelingModeAssetGenerationBehavior::AutoGenerateAndAutosave;
 	}
 	else if (NewString == AssetSaveModes[1])
 	{
-		Settings->AssetGenerationMode = EModelingModeAssetGenerationBehavior::AutoGenerateButDoNotAutosave;
+		AssetGenerationBehavior = EModelingModeAssetGenerationBehavior::AutoGenerateButDoNotAutosave;
 	}
 	else if (NewString == AssetSaveModes[2])
 	{
-		Settings->AssetGenerationMode = EModelingModeAssetGenerationBehavior::InteractivePromptToSave;
-	}
-	else
-	{
-		Settings->AssetGenerationMode = EModelingModeAssetGenerationBehavior::AutoGenerateButDoNotAutosave;
+		AssetGenerationBehavior = EModelingModeAssetGenerationBehavior::InteractivePromptToSave;
 	}
 
+	Settings->SetAssetGenerationMode(AssetGenerationBehavior);
+
 	Settings->SaveConfig();
 }
 
@@ -1044,7 +1066,7 @@ void FModelingToolsEditorModeToolkit::UpdateAssetPanelFromSettings()
 {
 	const UModelingToolsEditorModeSettings* Settings = GetDefault<UModelingToolsEditorModeSettings>();
 
-	switch (Settings->AssetGenerationLocation)
+	switch (Settings->GetAssetGenerationLocation())
 	{
 	case EModelingModeAssetGenerationLocation::CurrentAssetBrowserPathIfAvailable:
 		AssetLocationMode->SetSelectedItem(AssetLocationModes[2]);
@@ -1058,7 +1080,9 @@ void FModelingToolsEditorModeToolkit::UpdateAssetPanelFromSettings()
 		break;
 	}
 
-	switch (Settings->AssetGenerationMode)
+	AssetLocationMode->SetEnabled(!Settings->InRestrictiveMode());
+
+	switch (Settings->GetAssetGenerationMode())
 	{
 	case EModelingModeAssetGenerationBehavior::AutoGenerateButDoNotAutosave:
 		AssetSaveMode->SetSelectedItem(AssetSaveModes[1]);
@@ -1079,10 +1103,10 @@ void FModelingToolsEditorModeToolkit::UpdateObjectCreationOptionsFromSettings()
 	const UModelingToolsEditorModeSettings* Settings = GetDefault<UModelingToolsEditorModeSettings>();
 
 	// enable/disable dynamic mesh actors
-	UCreateMeshObjectTypeProperties::bEnableDynamicMeshActorSupport = Settings->bEnableDynamicMeshActors;
+	UCreateMeshObjectTypeProperties::bEnableDynamicMeshActorSupport = Settings->GetEnableDynamicMeshActors();
 
 	// set configured default type
-	if (Settings->DefaultMeshObjectType == EModelingModeDefaultMeshObjectType::DynamicMeshActor && Settings->bEnableDynamicMeshActors)
+	if (Settings->DefaultMeshObjectType == EModelingModeDefaultMeshObjectType::DynamicMeshActor && Settings->GetEnableDynamicMeshActors())
 	{
 		UCreateMeshObjectTypeProperties::DefaultObjectTypeIdentifier = UCreateMeshObjectTypeProperties::DynamicMeshActorIdentifier;
 	}
@@ -1110,5 +1134,17 @@ void FModelingToolsEditorModeToolkit::OnShowAssetSettings()
 	}
 }
 
+FText FModelingToolsEditorModeToolkit::GetRestrictiveModeAutogenSubfolderText() const
+{
+	const UModelingToolsEditorModeSettings* Settings = GetDefault<UModelingToolsEditorModeSettings>();
+	return FText::FromString(Settings->GetRestrictiveModeAutogenSubfolder());
+}
+
+void FModelingToolsEditorModeToolkit::OnRestrictiveModeAutogenSubfolderTextCommitted(const FText& InNewText, ETextCommit::Type InTextCommit) const
+{
+	UModelingToolsEditorModeSettings* Settings = GetMutableDefault<UModelingToolsEditorModeSettings>();
+	Settings->SetRestrictiveModeAutogenSubfolder(InNewText.ToString());
+}
+
 
 #undef LOCTEXT_NAMESPACE
diff --git a/Engine/Plugins/Editor/ModelingToolsEditorMode/Source/ModelingToolsEditorMode/Public/ModelingToolsEditorModeSettings.h b/Engine/Plugins/Editor/ModelingToolsEditorMode/Source/ModelingToolsEditorMode/Public/ModelingToolsEditorModeSettings.h
--- a/Engine/Plugins/Editor/ModelingToolsEditorMode/Source/ModelingToolsEditorMode/Public/ModelingToolsEditorModeSettings.h
+++ b/Engine/Plugins/Editor/ModelingToolsEditorMode/Source/ModelingToolsEditorMode/Public/ModelingToolsEditorModeSettings.h
@@ -62,30 +62,63 @@ public:
 
 	// UDeveloperSettings overrides
 
-	virtual FName GetContainerName() const { return FName("Project"); }
-	virtual FName GetCategoryName() const { return FName("Plugins"); }
-	virtual FName GetSectionName() const { return FName("ModelingMode"); }
+	virtual FName GetContainerName() const override { return FName("Project"); }
+	virtual FName GetCategoryName() const override { return FName("Plugins"); }
+	virtual FName GetSectionName() const override { return FName("ModelingMode"); }
 
 	virtual FText GetSectionText() const override;
 	virtual FText GetSectionDescription() const override;
 
-public:
+protected:
 
 	/** Enable/Disable the options to emit Dynamic Mesh Actors in Modeling Mode Tools */
 	UPROPERTY(config, EditAnywhere, Category = "Modeling Mode")
 	bool bEnableDynamicMeshActors = false;
 
-	/** What type of Mesh Object should Output Type Setting default to in Modeling Mode Tools (takes effect after Editor restart) */
-	UPROPERTY(config, EditAnywhere, Category = "Modeling Mode")
-	EModelingModeDefaultMeshObjectType DefaultMeshObjectType = EModelingModeDefaultMeshObjectType::StaticMeshAsset;
+	/** Where should Assets auto-generated by Modeling Tools be stored by default */
+	UPROPERTY(config, EditAnywhere, Category = "Modeling Mode|Generated Assets")
+	EModelingModeAssetGenerationLocation AssetGenerationLocation = EModelingModeAssetGenerationLocation::AutoGeneratedWorldRelativeAssetPath;
 
 	/** How should Assets auto-generated by Modeling Tools be handled in terms of saving, naming, etc */
 	UPROPERTY(config, EditAnywhere, Category = "Modeling Mode|Generated Assets")
 	EModelingModeAssetGenerationBehavior AssetGenerationMode = EModelingModeAssetGenerationBehavior::AutoGenerateButDoNotAutosave;
 
-	/** Where should Assets auto-generated by Modeling Tools be stored by default */
-	UPROPERTY(config, EditAnywhere, Category = "Modeling Mode|Generated Assets")
-	EModelingModeAssetGenerationLocation AssetGenerationLocation = EModelingModeAssetGenerationLocation::AutoGeneratedWorldRelativeAssetPath;
+public:
+
+	bool GetEnableDynamicMeshActors() const
+	{
+		return bRestrictiveMode
+					? true
+					: bEnableDynamicMeshActors;
+	}
+
+	EModelingModeAssetGenerationLocation GetAssetGenerationLocation() const
+	{
+		return bRestrictiveMode
+				   ? EModelingModeAssetGenerationLocation::AutoGeneratedWorldRelativeAssetPath
+				   : AssetGenerationLocation;
+	}
+
+	void SetAssetGenerationLocation(const EModelingModeAssetGenerationLocation Location)
+	{
+		AssetGenerationLocation = Location;
+	}
+
+	EModelingModeAssetGenerationBehavior GetAssetGenerationMode() const
+	{
+		return bRestrictiveMode
+					? EModelingModeAssetGenerationBehavior::AutoGenerateButDoNotAutosave
+					: AssetGenerationMode;
+	}
+
+	void SetAssetGenerationMode(const EModelingModeAssetGenerationBehavior Mode)
+	{
+		AssetGenerationMode = Mode;
+	}
+
+	/** What type of Mesh Object should Output Type Setting default to in Modeling Mode Tools (takes effect after Editor restart) */
+	UPROPERTY(config, EditAnywhere, Category = "Modeling Mode")
+	EModelingModeDefaultMeshObjectType DefaultMeshObjectType = EModelingModeDefaultMeshObjectType::StaticMeshAsset;
 
 	/** Assets auto-generated by Modeling Tools are stored at this path, relative to the parent path defined by the Location. Set to an empty string to disable. */
 	UPROPERTY(config, EditAnywhere, Category = "Modeling Mode|Generated Assets")
@@ -107,18 +140,14 @@ public:
 	UPROPERTY(config, EditAnywhere, Category = "Modeling Mode|Generated Assets")
 	bool bAppendRandomStringToName = true;
 
-
 	//
 	// Selection
 	//
 
-
 	/** If true, selections will be tracked between tools, and visualized with outlines. This feature is Experimental, disabled by default. */
 	UPROPERTY(config, EditAnywhere, Category = "Modeling Mode|Selection")
 	bool bEnablePersistentSelections = false;
 
-public:
-
 	DECLARE_MULTICAST_DELEGATE_TwoParams(UModelingToolsEditorModeSettingsModified, UObject*, FProperty*);
 	UModelingToolsEditorModeSettingsModified OnModified;
 
@@ -129,6 +158,43 @@ public:
 		Super::PostEditChangeProperty(PropertyChangedEvent);
 	}
 
+	//
+	// Restrictive Mode
+	//
+
+	virtual bool InRestrictiveMode() const
+	{
+		return bRestrictiveMode;
+	}
+
+	virtual void SetRestrictiveMode(bool bEnabled)
+	{
+		bRestrictiveMode = bEnabled;
+	}
+
+	virtual FString GetRestrictiveModeAutogenSubfolder() const
+	{
+		return bRestrictiveMode ? RestrictiveModeAutogenSubfolder : FString();
+	}
+
+	virtual bool SetRestrictiveModeAutogenSubfolder(const FString& AutogenSubfolder)
+	{
+		if (bRestrictiveMode)
+		{
+			RestrictiveModeAutogenSubfolder = AutogenSubfolder;
+			return true;
+		}
+		return false;
+	}
+
+protected:
+
+	/** Override for the name of the subfolder in which new assets get stored in when in restrictive mode. */
+	UPROPERTY(Config)
+	FString RestrictiveModeAutogenSubfolder;
+
+	bool bRestrictiveMode = false;
+
 };
 
 
diff --git a/Engine/Plugins/Editor/ModelingToolsEditorMode/Source/ModelingToolsEditorMode/Public/ModelingToolsEditorModeToolkit.h b/Engine/Plugins/Editor/ModelingToolsEditorMode/Source/ModelingToolsEditorMode/Public/ModelingToolsEditorModeToolkit.h
--- a/Engine/Plugins/Editor/ModelingToolsEditorMode/Source/ModelingToolsEditorMode/Public/ModelingToolsEditorModeToolkit.h
+++ b/Engine/Plugins/Editor/ModelingToolsEditorMode/Source/ModelingToolsEditorMode/Public/ModelingToolsEditorModeToolkit.h
@@ -62,10 +62,12 @@ public:
 
 	void OnActiveViewportChanged(TSharedPtr<IAssetViewport>, TSharedPtr<IAssetViewport> );
 
-
-
 	virtual void InvokeUI() override;
 
+	// Read and write Autogen Subfolder quick setting in restrictive mode.
+	virtual FText GetRestrictiveModeAutogenSubfolderText() const;
+	virtual void OnRestrictiveModeAutogenSubfolderTextCommitted(const FText& InNewText, ETextCommit::Type InTextCommit) const;
+
 private:
 	const static TArray<FName> PaletteNames_Standard;
 	const static TArray<FName> PaletteNames_Experimental;
@@ -97,6 +99,7 @@ private:
 	TArray<TSharedPtr<FString>> AssetSaveModes;
 	TSharedPtr<STextComboBox> AssetLocationMode;
 	TSharedPtr<STextComboBox> AssetSaveMode;
+	TSharedPtr<SEditableTextBox> AutogenSubfolder;
 	void UpdateAssetLocationMode(TSharedPtr<FString> NewString);
 	void UpdateAssetSaveMode(TSharedPtr<FString> NewString);
 	void UpdateAssetPanelFromSettings();
