ModelingToolset: Fix missing include

diff --git a/Engine/Plugins/Editor/ModelingToolsEditorMode/Source/ModelingToolsEditorMode/Public/ModelingToolsEditorModeToolkit.h b/Engine/Plugins/Editor/ModelingToolsEditorMode/Source/ModelingToolsEditorMode/Public/ModelingToolsEditorModeToolkit.h
--- a/Engine/Plugins/Editor/ModelingToolsEditorMode/Source/ModelingToolsEditorMode/Public/ModelingToolsEditorModeToolkit.h
+++ b/Engine/Plugins/Editor/ModelingToolsEditorMode/Source/ModelingToolsEditorMode/Public/ModelingToolsEditorModeToolkit.h
@@ -8,6 +8,7 @@
 #include "ModelingToolsEditorMode.h"
 #include "Widgets/Input/SComboBox.h"
 #include "Widgets/Input/STextComboBox.h"
+#include "Widgets/Input/SEditableTextBox.h"
 #include "Widgets/SBoxPanel.h"
 #include "StatusBarSubsystem.h"
 
