MeshModelingToolset: Allow users to explicitly choose path for generated assets inside the project folder when in restrictive mode (defaulted to "Meshes")

diff --git a/Engine/Plugins/Editor/ModelingToolsEditorMode/Source/ModelingToolsEditorMode/Private/ModelingModeAssetUtils.cpp b/Engine/Plugins/Editor/ModelingToolsEditorMode/Source/ModelingToolsEditorMode/Private/ModelingModeAssetUtils.cpp
--- a/Engine/Plugins/Editor/ModelingToolsEditorMode/Source/ModelingToolsEditorMode/Private/ModelingModeAssetUtils.cpp
+++ b/Engine/Plugins/Editor/ModelingToolsEditorMode/Source/ModelingToolsEditorMode/Private/ModelingModeAssetUtils.cpp
@@ -22,18 +22,6 @@ namespace UE
 namespace Local
 {
 
-FString GetWorldRelativeAssetRootPath(const UWorld* World)
-{
-	if (ensure(World->GetOutermost() != nullptr) == false)
-	{
-		return TEXT("/Game/");
-	}
-	FString WorldPackageName = World->GetOutermost()->GetName();
-	FString WorldPackageFolder = FPackageName::GetLongPackagePath(WorldPackageName);
-	return WorldPackageFolder;
-}
-
-
 FString GetActiveAssetFolderPath()
 {
 	IContentBrowserSingleton& ContentBrowser = FModuleManager::LoadModuleChecked<FContentBrowserModule>("ContentBrowser").Get();
@@ -76,9 +64,22 @@ FString InteractiveSelectAssetPath(const FString& DefaultAssetName, const FText&
 }
 
 
+FString UE::Modeling::GetWorldRelativeAssetRootPath(const UWorld* World)
+{
+	if (World == nullptr)
+	{
+		World = GEditor->bIsSimulatingInEditor ? GEditor->GetPIEWorldContext()->World() : GEditor->GetEditorWorldContext().World();
+	}
 
+	if (ensure(World->GetOutermost() != nullptr) == false)
+	{
+		return TEXT("/Game/");
+	}
 
-
+	const FString WorldPackageName = World->GetOutermost()->GetName();
+	const FString WorldPackageFolder = FPackageName::GetLongPackagePath(WorldPackageName);
+	return WorldPackageFolder;
+}
 
 FString UE::Modeling::GetNewAssetPathName(const FString& BaseNameIn, const UWorld* TargetWorld, FString SuggestedFolder)
 {
@@ -98,7 +99,7 @@ FString UE::Modeling::GetNewAssetPathName(const FString& BaseNameIn, const UWorl
 		// be the plugin top-level path and not /Game.
 		if (Settings->GetAssetGenerationLocation() == EModelingModeAssetGenerationLocation::AutoGeneratedWorldRelativeAssetPath)
 		{
-			PackageFolderPath = UE::Local::GetWorldRelativeAssetRootPath(TargetWorld);
+			PackageFolderPath = GetWorldRelativeAssetRootPath(TargetWorld);
 		}
 
 		// If the level is unsaved, then GetWorldRelativeAssetRootPath will return "/Temp". This is an undesirable
@@ -110,30 +111,23 @@ FString UE::Modeling::GetNewAssetPathName(const FString& BaseNameIn, const UWorl
 			PackageFolderPath = TEXT("/Game/");
 		}
 
-		// combine with fixed AutoGen path name if it is not empty
-		if (Settings->AutoGeneratedAssetPath.Len() > 0)
-		{
-			PackageFolderPath = FPaths::Combine(PackageFolderPath, Settings->AutoGeneratedAssetPath);
-		}
-
 		if (Settings->InRestrictiveMode())
 		{
-			const FString AutogenSubfolder = Settings->GetRestrictiveModeAutogenSubfolder().TrimStartAndEnd();
-			if (!AutogenSubfolder.IsEmpty())
+			// combine with fixed AutoGen path name if it is not empty
+			const FString RestrictiveModeAutogenAssetPath = Settings->GetRestrictiveModeAutoGeneratedAssetPath(); 
+			if (RestrictiveModeAutogenAssetPath.Len() > 0)
 			{
-				PackageFolderPath = FPaths::Combine(PackageFolderPath, AutogenSubfolder);
-			}
-			else
-			{
-				FString UsernameString = FPlatformProcess::UserName();
-				if (!UsernameString.IsEmpty())
-				{
-					PackageFolderPath = FPaths::Combine(PackageFolderPath, UsernameString);
-				}
+				PackageFolderPath = FPaths::Combine(PackageFolderPath, RestrictiveModeAutogenAssetPath);
 			}
 		}
 		else
 		{
+			// combine with fixed AutoGen path name if it is not empty
+			if (Settings->AutoGeneratedAssetPath.Len() > 0)
+			{
+				PackageFolderPath = FPaths::Combine(PackageFolderPath, Settings->AutoGeneratedAssetPath);
+			}
+
 			// append username-specific subfolder
 			if (Settings->bUsePerUserAutogenSubfolder)
 			{
diff --git a/Engine/Plugins/Editor/ModelingToolsEditorMode/Source/ModelingToolsEditorMode/Private/ModelingToolsEditorModeToolkit.cpp b/Engine/Plugins/Editor/ModelingToolsEditorMode/Source/ModelingToolsEditorMode/Private/ModelingToolsEditorModeToolkit.cpp
--- a/Engine/Plugins/Editor/ModelingToolsEditorMode/Source/ModelingToolsEditorMode/Private/ModelingToolsEditorModeToolkit.cpp
+++ b/Engine/Plugins/Editor/ModelingToolsEditorMode/Source/ModelingToolsEditorMode/Private/ModelingToolsEditorModeToolkit.cpp
@@ -9,9 +9,7 @@
 
 #include "Framework/MultiBox/MultiBoxBuilder.h"
 #include "Modules/ModuleManager.h"
-#include "PropertyEditorModule.h"
 #include "IDetailsView.h"
-#include "IDetailRootObjectCustomization.h"
 #include "ISettingsModule.h"
 #include "EditorModeManager.h"
 #include "Toolkits/AssetEditorModeUILayer.h"
@@ -33,12 +31,14 @@
 // for Object Type properties
 #include "PropertySets/CreateMeshObjectTypeProperties.h"
 
-// for LOD setting
+// for quick settings
 #include "EditorInteractiveToolsFrameworkModule.h"
 #include "Tools/EditorComponentSourceFactory.h"
 #include "ToolTargetManager.h"
 #include "ToolTargets/StaticMeshComponentToolTarget.h"
 #include "SPrimaryButton.h"
+#include "Dialogs/DlgPickPath.h"
+#include "ModelingModeAssetUtils.h"
 
 
 #define LOCTEXT_NAMESPACE "FModelingToolsEditorModeToolkit"
@@ -321,22 +321,27 @@ TSharedPtr<SWidget> FModelingToolsEditorModeToolkit::MakeAssetConfigPanel()
 
 	if (Settings->InRestrictiveMode())
 	{
-		AutogenSubfolder = SNew(SEditableTextBox)
-		.Text(this, &FModelingToolsEditorModeToolkit::GetRestrictiveModeAutogenSubfolderText)
-		.OnTextCommitted(this, &FModelingToolsEditorModeToolkit::OnRestrictiveModeAutogenSubfolderTextCommitted);
-		
+		NewAssetPath = SNew(SEditableTextBox)
+		.Text(this, &FModelingToolsEditorModeToolkit::GetRestrictiveModeAutoGeneratedAssetPathText)
+		.OnTextCommitted(this, &FModelingToolsEditorModeToolkit::OnRestrictiveModeAutoGeneratedAssetPathTextCommitted);
+
 		Content->AddSlot()[
 			SNew(SHorizontalBox)
 			+ SHorizontalBox::Slot().Padding(0, 2, 2, 2).HAlign(HAlign_Right).VAlign(VAlign_Center).AutoWidth()
 			[
 				SNew(STextBlock)
-				.Text(LOCTEXT("AutogenSubfolder", "Autogen Subfolder"))
-				.ToolTipText(LOCTEXT("AutogenSubfolderToolTip", "Override for the name of the subfolder in which new assets get stored in. "
-					"If nothing is set, the username is used as the subfolder name."))
+				.Text(LOCTEXT("NewAssetPathLabel", "New Asset Path"))
+				.ToolTipText(LOCTEXT("NewAssetPathToolTip", "Path used for storing newly generated assets within the project folder."))
 			]
 			+ SHorizontalBox::Slot().HAlign(HAlign_Fill).Padding(0).FillWidth(1.0f)
 			[
-				AutogenSubfolder->AsShared()
+				NewAssetPath->AsShared()
+			]
+			+ SHorizontalBox::Slot().HAlign(HAlign_Right).Padding(0, 0, 0, 0).AutoWidth()
+			[
+				SNew(SSimpleButton)
+				.OnClicked_Lambda( [this]() { SelectNewAssetPath(); return FReply::Handled(); } )
+				.Icon(FAppStyle::Get().GetBrush("Icons.FolderOpen"))
 			]
 		];
 	}
@@ -1134,16 +1139,43 @@ void FModelingToolsEditorModeToolkit::OnShowAssetSettings()
 	}
 }
 
-FText FModelingToolsEditorModeToolkit::GetRestrictiveModeAutogenSubfolderText() const
+void FModelingToolsEditorModeToolkit::SelectNewAssetPath() const
+{	
+	const FString PackageFolderPath = []
+	{
+		FString Path = UE::Modeling::GetWorldRelativeAssetRootPath(nullptr);
+		Path.RemoveFromEnd("/");
+		return Path;
+	}();
+	
+	const TSharedPtr<SDlgPickPath> PickPathWidget =
+		SNew(SDlgPickPath)
+		.Title(FText::FromString("Choose New Asset Path"))
+		.DefaultPath(FText::FromString(PackageFolderPath + '/' + GetRestrictiveModeAutoGeneratedAssetPathText().ToString()))
+		.AllowReadOnlyFolders(false);
+
+	if (PickPathWidget->ShowModal() == EAppReturnType::Ok)
+	{
+		FString Path = PickPathWidget->GetPath().ToString();
+		Path.RemoveFromStart(PackageFolderPath);	// Remove project folder prefix
+		Path.RemoveFromStart("/");			// Remove leading '/' in case the path is the project folder
+		
+		UModelingToolsEditorModeSettings* Settings = GetMutableDefault<UModelingToolsEditorModeSettings>();
+		Settings->SetRestrictiveModeAutoGeneratedAssetPath(Path);
+	}
+}
+
+
+FText FModelingToolsEditorModeToolkit::GetRestrictiveModeAutoGeneratedAssetPathText() const
 {
 	const UModelingToolsEditorModeSettings* Settings = GetDefault<UModelingToolsEditorModeSettings>();
-	return FText::FromString(Settings->GetRestrictiveModeAutogenSubfolder());
+	return FText::FromString(Settings->GetRestrictiveModeAutoGeneratedAssetPath());
 }
 
-void FModelingToolsEditorModeToolkit::OnRestrictiveModeAutogenSubfolderTextCommitted(const FText& InNewText, ETextCommit::Type InTextCommit) const
+void FModelingToolsEditorModeToolkit::OnRestrictiveModeAutoGeneratedAssetPathTextCommitted(const FText& InNewText, ETextCommit::Type InTextCommit) const
 {
 	UModelingToolsEditorModeSettings* Settings = GetMutableDefault<UModelingToolsEditorModeSettings>();
-	Settings->SetRestrictiveModeAutogenSubfolder(InNewText.ToString());
+	Settings->SetRestrictiveModeAutoGeneratedAssetPath(InNewText.ToString());
 }
 
 
diff --git a/Engine/Plugins/Editor/ModelingToolsEditorMode/Source/ModelingToolsEditorMode/Public/ModelingModeAssetUtils.h b/Engine/Plugins/Editor/ModelingToolsEditorMode/Source/ModelingToolsEditorMode/Public/ModelingModeAssetUtils.h
--- a/Engine/Plugins/Editor/ModelingToolsEditorMode/Source/ModelingToolsEditorMode/Public/ModelingModeAssetUtils.h
+++ b/Engine/Plugins/Editor/ModelingToolsEditorMode/Source/ModelingToolsEditorMode/Public/ModelingModeAssetUtils.h
@@ -26,6 +26,12 @@ namespace Modeling
  */
 MODELINGTOOLSEDITORMODE_API FString GetNewAssetPathName(const FString& BaseName, const UWorld* TargetWorld, FString SuggestedFolder);
 
+/**
+ * Determines the world relative asset root path.
+ * @param World World the new Asset will be used in; if this is a nullptr then it will be set using the current world context.
+ * @return Asset root path relative to the world.
+ */
+MODELINGTOOLSEDITORMODE_API FString GetWorldRelativeAssetRootPath(const UWorld* World);
 
 /**
  * Utility function that may auto-save a newly created asset depending on current mode settings,
diff --git a/Engine/Plugins/Editor/ModelingToolsEditorMode/Source/ModelingToolsEditorMode/Public/ModelingToolsEditorModeSettings.h b/Engine/Plugins/Editor/ModelingToolsEditorMode/Source/ModelingToolsEditorMode/Public/ModelingToolsEditorModeSettings.h
--- a/Engine/Plugins/Editor/ModelingToolsEditorMode/Source/ModelingToolsEditorMode/Public/ModelingToolsEditorModeSettings.h
+++ b/Engine/Plugins/Editor/ModelingToolsEditorMode/Source/ModelingToolsEditorMode/Public/ModelingToolsEditorModeSettings.h
@@ -172,16 +172,16 @@ public:
 		bRestrictiveMode = bEnabled;
 	}
 
-	virtual FString GetRestrictiveModeAutogenSubfolder() const
+	virtual FString GetRestrictiveModeAutoGeneratedAssetPath() const
 	{
-		return bRestrictiveMode ? RestrictiveModeAutogenSubfolder : FString();
+		return bRestrictiveMode ? RestrictiveModeAutoGeneratedAssetPath : FString();
 	}
 
-	virtual bool SetRestrictiveModeAutogenSubfolder(const FString& AutogenSubfolder)
+	virtual bool SetRestrictiveModeAutoGeneratedAssetPath(const FString& AssetPath)
 	{
 		if (bRestrictiveMode)
 		{
-			RestrictiveModeAutogenSubfolder = AutogenSubfolder;
+			RestrictiveModeAutoGeneratedAssetPath = AssetPath;
 			return true;
 		}
 		return false;
@@ -189,9 +189,9 @@ public:
 
 protected:
 
-	/** Override for the name of the subfolder in which new assets get stored in when in restrictive mode. */
+	/** Assets auto-generated by Modeling Tools are stored at this path when in restrictive mode, relative to the package folder path. */
 	UPROPERTY(Config)
-	FString RestrictiveModeAutogenSubfolder;
+	FString RestrictiveModeAutoGeneratedAssetPath = TEXT("Meshes");
 
 	bool bRestrictiveMode = false;
 
diff --git a/Engine/Plugins/Editor/ModelingToolsEditorMode/Source/ModelingToolsEditorMode/Public/ModelingToolsEditorModeToolkit.h b/Engine/Plugins/Editor/ModelingToolsEditorMode/Source/ModelingToolsEditorMode/Public/ModelingToolsEditorModeToolkit.h
--- a/Engine/Plugins/Editor/ModelingToolsEditorMode/Source/ModelingToolsEditorMode/Public/ModelingToolsEditorModeToolkit.h
+++ b/Engine/Plugins/Editor/ModelingToolsEditorMode/Source/ModelingToolsEditorMode/Public/ModelingToolsEditorModeToolkit.h
@@ -4,12 +4,9 @@
 
 #include "CoreMinimal.h"
 #include "Toolkits/BaseToolkit.h"
-#include "InteractiveTool.h"
 #include "ModelingToolsEditorMode.h"
-#include "Widgets/Input/SComboBox.h"
 #include "Widgets/Input/STextComboBox.h"
 #include "Widgets/Input/SEditableTextBox.h"
-#include "Widgets/SBoxPanel.h"
 #include "StatusBarSubsystem.h"
 
 
@@ -66,8 +63,8 @@ public:
 	virtual void InvokeUI() override;
 
 	// Read and write Autogen Subfolder quick setting in restrictive mode.
-	virtual FText GetRestrictiveModeAutogenSubfolderText() const;
-	virtual void OnRestrictiveModeAutogenSubfolderTextCommitted(const FText& InNewText, ETextCommit::Type InTextCommit) const;
+	virtual FText GetRestrictiveModeAutoGeneratedAssetPathText() const;
+	virtual void OnRestrictiveModeAutoGeneratedAssetPathTextCommitted(const FText& InNewText, ETextCommit::Type InTextCommit) const;
 
 private:
 	const static TArray<FName> PaletteNames_Standard;
@@ -76,7 +73,7 @@ private:
 	FText ActiveToolName;
 	FText ActiveToolMessage;
 	FStatusBarMessageHandle ActiveToolMessageHandle;
-	const FSlateBrush* ActiveToolIcon;
+	const FSlateBrush* ActiveToolIcon = nullptr;
 
 	TSharedPtr<SWidget> ToolkitWidget;
 	void UpdateActiveToolProperties();
@@ -100,7 +97,7 @@ private:
 	TArray<TSharedPtr<FString>> AssetSaveModes;
 	TSharedPtr<STextComboBox> AssetLocationMode;
 	TSharedPtr<STextComboBox> AssetSaveMode;
-	TSharedPtr<SEditableTextBox> AutogenSubfolder;
+	TSharedPtr<SEditableTextBox> NewAssetPath;
 	void UpdateAssetLocationMode(TSharedPtr<FString> NewString);
 	void UpdateAssetSaveMode(TSharedPtr<FString> NewString);
 	void UpdateAssetPanelFromSettings();
@@ -108,6 +105,7 @@ private:
 	void OnAssetSettingsModified();
 	FDelegateHandle AssetSettingsModifiedHandle;
 	void OnShowAssetSettings();
+	void SelectNewAssetPath() const;
 
 	TArray<TSharedPtr<FString>> AssetLODModes;
 	TSharedPtr<STextBlock> AssetLODModeLabel;
diff --git a/Engine/Source/Editor/UnrealEd/Private/Dialogs/DlgPickPath.cpp b/Engine/Source/Editor/UnrealEd/Private/Dialogs/DlgPickPath.cpp
--- a/Engine/Source/Editor/UnrealEd/Private/Dialogs/DlgPickPath.cpp
+++ b/Engine/Source/Editor/UnrealEd/Private/Dialogs/DlgPickPath.cpp
@@ -17,11 +17,13 @@
 void SDlgPickPath::Construct(const FArguments& InArgs)
 {
 	Path = FText::FromString(FPackageName::GetLongPackagePath(InArgs._DefaultPath.ToString()));
+	bAllowReadOnlyFolders = InArgs._AllowReadOnlyFolders;
 
 	FPathPickerConfig PathPickerConfig;
 	PathPickerConfig.DefaultPath = Path.ToString();
 	PathPickerConfig.OnPathSelected = FOnPathSelected::CreateSP(this, &SDlgPickPath::OnPathChange);
 	PathPickerConfig.bAddDefaultPath = true;
+	PathPickerConfig.bAllowReadOnlyFolders = bAllowReadOnlyFolders;
 
 	FContentBrowserModule& ContentBrowserModule = FModuleManager::LoadModuleChecked<FContentBrowserModule>("ContentBrowser");
 
diff --git a/Engine/Source/Editor/UnrealEd/Public/Dialogs/DlgPickPath.h b/Engine/Source/Editor/UnrealEd/Public/Dialogs/DlgPickPath.h
--- a/Engine/Source/Editor/UnrealEd/Public/Dialogs/DlgPickPath.h
+++ b/Engine/Source/Editor/UnrealEd/Public/Dialogs/DlgPickPath.h
@@ -13,10 +13,15 @@ class SDlgPickPath : public SWindow
 {
 public:
 	SLATE_BEGIN_ARGS(SDlgPickPath)
-	{
-	}
+		: _AllowReadOnlyFolders(true)
+		{}
+	
 		SLATE_ARGUMENT(FText, Title)
 		SLATE_ARGUMENT(FText, DefaultPath)
+
+		/** If true, read only folders will be displayed */
+		SLATE_ARGUMENT(bool, AllowReadOnlyFolders)
+
 	SLATE_END_ARGS()
 
 	UNREALED_API SDlgPickPath()
@@ -40,6 +45,9 @@ protected:
 
 	EAppReturnType::Type UserResponse;
 	FText Path;
+
+	/** If true, read only folders will be displayed */
+	bool bAllowReadOnlyFolders;
 };
 
 #undef LOCTEXT_NAMESPACE
