Fixed non-unity compile errors

diff --git a/Engine/Plugins/Editor/ModelingToolsEditorMode/Source/ModelingToolsEditorMode/Private/ModelingModeAssetUtils.cpp b/Engine/Plugins/Editor/ModelingToolsEditorMode/Source/ModelingToolsEditorMode/Private/ModelingModeAssetUtils.cpp
--- a/Engine/Plugins/Editor/ModelingToolsEditorMode/Source/ModelingToolsEditorMode/Private/ModelingModeAssetUtils.cpp
+++ b/Engine/Plugins/Editor/ModelingToolsEditorMode/Source/ModelingToolsEditorMode/Private/ModelingModeAssetUtils.cpp
@@ -14,6 +14,7 @@
 #include "ContentBrowserModule.h"
 #include "IContentBrowserSingleton.h"
 
+extern UNREALED_API UEditorEngine* GEditor;
 
 #define LOCTEXT_NAMESPACE "ModelingModeAssetUtils"
 
