MeshModelingToolset: Fix non-unity build error in ModelingModeAssetUtil.cpp

diff --git a/Engine/Plugins/Editor/ModelingToolsEditorMode/Source/ModelingToolsEditorMode/Private/ModelingModeAssetUtils.cpp b/Engine/Plugins/Editor/ModelingToolsEditorMode/Source/ModelingToolsEditorMode/Private/ModelingModeAssetUtils.cpp
--- a/Engine/Plugins/Editor/ModelingToolsEditorMode/Source/ModelingToolsEditorMode/Private/ModelingModeAssetUtils.cpp
+++ b/Engine/Plugins/Editor/ModelingToolsEditorMode/Source/ModelingToolsEditorMode/Private/ModelingModeAssetUtils.cpp
@@ -9,6 +9,7 @@
 #include "FileHelpers.h"
 #include "Misc/Paths.h"
 #include "ModelingObjectsCreationAPI.h"
+#include "Editor/EditorEngine.h"
 
 // for content-browser things
 #include "ContentBrowserModule.h"
