Prevent BakeRS UI from allowing user to try baking only rotation w/ a rotation+non-uniform scale, as UE transforms don't support this

diff --git a/Engine/Plugins/Editor/ModelingToolsEditorMode/Source/ModelingToolsEditorMode/ModelingToolsEditorMode.Build.cs b/Engine/Plugins/Editor/ModelingToolsEditorMode/Source/ModelingToolsEditorMode/ModelingToolsEditorMode.Build.cs
--- a/Engine/Plugins/Editor/ModelingToolsEditorMode/Source/ModelingToolsEditorMode/ModelingToolsEditorMode.Build.cs
+++ b/Engine/Plugins/Editor/ModelingToolsEditorMode/Source/ModelingToolsEditorMode/ModelingToolsEditorMode.Build.cs
@@ -58,6 +58,7 @@ public class ModelingToolsEditorMode : ModuleRules
 				"ModelingEditorUI",
 				"StylusInput",
 				"DeveloperSettings",
+				"PropertyEditor",
 				// ... add private dependencies that you statically link with here ...	
 			}
 		);
diff --git a/Engine/Plugins/Editor/ModelingToolsEditorMode/Source/ModelingToolsEditorMode/Private/DetailsCustomizations/BakeTransformToolCustomizations.cpp b/Engine/Plugins/Editor/ModelingToolsEditorMode/Source/ModelingToolsEditorMode/Private/DetailsCustomizations/BakeTransformToolCustomizations.cpp
new file mode 100644
--- /dev/null
+++ b/Engine/Plugins/Editor/ModelingToolsEditorMode/Source/ModelingToolsEditorMode/Private/DetailsCustomizations/BakeTransformToolCustomizations.cpp
@@ -0,0 +1,63 @@
+// Copyright Epic Games, Inc. All Rights Reserved.
+
+#include "DetailsCustomizations/BakeTransformToolCustomizations.h"
+#include "UObject/Class.h"
+#include "DetailLayoutBuilder.h"
+#include "PropertyHandle.h"
+#include "PropertyRestriction.h"
+
+#include "BakeTransformTool.h"
+
+
+#define LOCTEXT_NAMESPACE "BakeTransformToolCustomizations"
+
+
+TSharedRef<IDetailCustomization> FBakeTransformToolDetails::MakeInstance()
+{
+	return MakeShareable(new FBakeTransformToolDetails);
+}
+
+void FBakeTransformToolDetails::CustomizeDetails(IDetailLayoutBuilder& DetailBuilder)
+{
+	TSharedRef<IPropertyHandle> AllowNoScale = DetailBuilder.GetProperty(GET_MEMBER_NAME_CHECKED(UBakeTransformToolProperties, bAllowNoScale), UBakeTransformToolProperties::StaticClass());
+	TSharedRef<IPropertyHandle> BakeScaleHandle = DetailBuilder.GetProperty(GET_MEMBER_NAME_CHECKED(UBakeTransformToolProperties, BakeScale), UBakeTransformToolProperties::StaticClass());
+	AllowNoScale->MarkHiddenByCustomization();
+
+	// Note: In practice, currently the tool computes AllowNoScale on Setup and the value does not change, so we don't need to monitor for property changes
+	// If AllowNoScale could change while the tool is open, we'd need to add a delegate like this:
+	//HasNonUniformScaleHandle->SetOnPropertyValueChanged(FSimpleDelegate::CreateRaw(this, &FBakeTransformToolDetails::EvaluateBakeScaleMethodRestrictions, AllowNoScale, BakeScaleHandle));
+
+	static FText RestrictReason = LOCTEXT("BakeScaleMethodRestrictionText", "cannot only bake rotation: selection includes meshes with both non-uniform scale and rotation");
+	BakeScaleMethodEnumRestriction = MakeShareable(new FPropertyRestriction(RestrictReason));
+	BakeScaleHandle->AddRestriction(BakeScaleMethodEnumRestriction.ToSharedRef());
+
+	EvaluateBakeScaleMethodRestrictions(AllowNoScale, BakeScaleHandle);
+}
+
+void FBakeTransformToolDetails::EvaluateBakeScaleMethodRestrictions(TSharedRef<IPropertyHandle> AllowNoScale, TSharedRef<IPropertyHandle> BakeScaleHandle)
+{
+	bool bAllowNoScale = false;
+	AllowNoScale->GetValue(bAllowNoScale);
+
+	UEnum* BakeScaleEnum = StaticEnum<EBakeScaleMethod>();
+
+	BakeScaleMethodEnumRestriction->RemoveAll();
+	if (!bAllowNoScale)
+	{
+		int32 EnumValue = BakeScaleEnum->GetValueByName("EBakeScaleMethod::DoNotBakeScale");
+		BakeScaleMethodEnumRestriction->AddDisabledValue("DoNotBakeScale");
+
+		// if the enum is set to DoNotBakeScale and we don't allow it, switch it over to BakeNonuniformScale instead
+		int32 Value = 0;
+		BakeScaleHandle->GetValue(Value);
+		if (Value == EnumValue)
+		{
+			Value = BakeScaleEnum->GetValueByName("EBakeScaleMethod::BakeNonuniformScale");
+			BakeScaleHandle->SetValue(Value, EPropertyValueSetFlags::NotTransactable | EPropertyValueSetFlags::InteractiveChange);
+		}
+	}
+
+
+}
+
+#undef LOCTEXT_NAMESPACE
diff --git a/Engine/Plugins/Editor/ModelingToolsEditorMode/Source/ModelingToolsEditorMode/Private/ModelingToolsEditorModeModule.cpp b/Engine/Plugins/Editor/ModelingToolsEditorMode/Source/ModelingToolsEditorMode/Private/ModelingToolsEditorModeModule.cpp
--- a/Engine/Plugins/Editor/ModelingToolsEditorMode/Source/ModelingToolsEditorMode/Private/ModelingToolsEditorModeModule.cpp
+++ b/Engine/Plugins/Editor/ModelingToolsEditorMode/Source/ModelingToolsEditorMode/Private/ModelingToolsEditorModeModule.cpp
@@ -15,12 +15,14 @@
 #include "DetailsCustomizations/ModelingToolsBrushSizeCustomization.h"
 #include "DetailsCustomizations/MeshVertexSculptToolCustomizations.h"
 #include "DetailsCustomizations/BakeMeshAttributeToolCustomizations.h"
+#include "DetailsCustomizations/BakeTransformToolCustomizations.h"
 
 #include "PropertySets/AxisFilterPropertyType.h"
 #include "MeshVertexSculptTool.h"
 #include "BakeMeshAttributeMapsTool.h"
 #include "BakeMultiMeshAttributeMapsTool.h"
 #include "BakeMeshAttributeVertexTool.h"
+#include "BakeTransformTool.h"
 
 
 #define LOCTEXT_NAMESPACE "FModelingToolsEditorModeModule"
@@ -91,6 +93,9 @@ void FModelingToolsEditorModeModule::OnPostEngineInit()
 	ClassesToUnregisterOnShutdown.Add(UBakeMultiMeshAttributeMapsToolProperties::StaticClass()->GetFName());
 	PropertyModule.RegisterCustomClassLayout("BakeMeshAttributeVertexToolProperties", FOnGetDetailCustomizationInstance::CreateStatic(&FBakeMeshAttributeVertexToolDetails::MakeInstance));
 	ClassesToUnregisterOnShutdown.Add(UBakeMeshAttributeVertexToolProperties::StaticClass()->GetFName());
+
+	PropertyModule.RegisterCustomClassLayout("BakeTransformToolProperties", FOnGetDetailCustomizationInstance::CreateStatic(&FBakeTransformToolDetails::MakeInstance));
+	ClassesToUnregisterOnShutdown.Add(UBakeTransformToolProperties::StaticClass()->GetFName());
 }
 
 #undef LOCTEXT_NAMESPACE
diff --git a/Engine/Plugins/Editor/ModelingToolsEditorMode/Source/ModelingToolsEditorMode/Public/DetailsCustomizations/BakeTransformToolCustomizations.h b/Engine/Plugins/Editor/ModelingToolsEditorMode/Source/ModelingToolsEditorMode/Public/DetailsCustomizations/BakeTransformToolCustomizations.h
new file mode 100644
--- /dev/null
+++ b/Engine/Plugins/Editor/ModelingToolsEditorMode/Source/ModelingToolsEditorMode/Public/DetailsCustomizations/BakeTransformToolCustomizations.h
@@ -0,0 +1,20 @@
+// Copyright Epic Games, Inc. All Rights Reserved.
+
+#pragma once
+
+#include "IDetailCustomization.h"
+#include "Templates/SharedPointer.h"
+
+// Customization for UBakeTransformToolProperties
+class FBakeTransformToolDetails : public IDetailCustomization
+{
+public:
+	static TSharedRef<IDetailCustomization> MakeInstance();
+	virtual void CustomizeDetails(IDetailLayoutBuilder& DetailBuilder) override;
+
+protected:
+	void EvaluateBakeScaleMethodRestrictions(TSharedRef<IPropertyHandle> AllowNoScale, TSharedRef<IPropertyHandle> BakeScaleHandle);
+
+	/** Property restriction applied to the EBakeScaleMethod enum dropdown box */
+	TSharedPtr<FPropertyRestriction> BakeScaleMethodEnumRestriction;
+};
diff --git a/Engine/Plugins/Experimental/MeshModelingToolsetExp/Source/MeshModelingToolsExp/Private/BakeTransformTool.cpp b/Engine/Plugins/Experimental/MeshModelingToolsetExp/Source/MeshModelingToolsExp/Private/BakeTransformTool.cpp
--- a/Engine/Plugins/Experimental/MeshModelingToolsetExp/Source/MeshModelingToolsExp/Private/BakeTransformTool.cpp
+++ b/Engine/Plugins/Experimental/MeshModelingToolsetExp/Source/MeshModelingToolsExp/Private/BakeTransformTool.cpp
@@ -71,10 +71,15 @@ void UBakeTransformTool::Setup()
 	}
 
 	bool bHasZeroScales = false;
+	BasicProperties->bAllowNoScale = true;
 	for (int32 ComponentIdx = 0; ComponentIdx < Targets.Num(); ComponentIdx++)
 	{
 		FTransform Transform = (FTransform)UE::ToolTarget::GetLocalToWorldTransform(Targets[ComponentIdx]);
-		if (Transform.GetScale3D().GetAbsMin() < KINDA_SMALL_NUMBER)
+		FVector Scale = Transform.GetScale3D();
+		// Set variable so a DetailCustomization can disable "DoNotBakeScale" if we have some targets w/ both rotation and non-uniform scale
+		// Note we could relax this to allow DoNotBakeScale if the rotation axis is aligned to the non-uniform scale axis, but that might make the enable/disable condition harder to understand.
+		BasicProperties->bAllowNoScale = BasicProperties->bAllowNoScale && (Transform.GetRotation().IsIdentity() || Transform.GetScale3D().AllComponentsEqual());
+		if (Scale.GetAbsMin() < KINDA_SMALL_NUMBER)
 		{
 			bHasZeroScales = true;
 		}
diff --git a/Engine/Plugins/Experimental/MeshModelingToolsetExp/Source/MeshModelingToolsExp/Public/BakeTransformTool.h b/Engine/Plugins/Experimental/MeshModelingToolsetExp/Source/MeshModelingToolsExp/Public/BakeTransformTool.h
--- a/Engine/Plugins/Experimental/MeshModelingToolsetExp/Source/MeshModelingToolsExp/Public/BakeTransformTool.h
+++ b/Engine/Plugins/Experimental/MeshModelingToolsetExp/Source/MeshModelingToolsExp/Public/BakeTransformTool.h
@@ -63,6 +63,11 @@ public:
 	/** Recenter pivot after baking transform */
 	UPROPERTY(EditAnywhere, Category = Options)
 	bool bRecenterPivot = false;
+
+	// This variable is used to enable/disable the "DoNotBakeScale" option in the BakeScale enum
+	// It is marked as visible only so that the FBakeTransformToolDetails can read it; it will be hidden by that DetailCustomization
+	UPROPERTY(VisibleAnywhere, Category = Options)
+	bool bAllowNoScale = true;
 };
 
 
