fix nonunity build error for bake transform tool detail customization

diff --git a/Engine/Plugins/Editor/ModelingToolsEditorMode/Source/ModelingToolsEditorMode/Public/DetailsCustomizations/BakeTransformToolCustomizations.h b/Engine/Plugins/Editor/ModelingToolsEditorMode/Source/ModelingToolsEditorMode/Public/DetailsCustomizations/BakeTransformToolCustomizations.h
--- a/Engine/Plugins/Editor/ModelingToolsEditorMode/Source/ModelingToolsEditorMode/Public/DetailsCustomizations/BakeTransformToolCustomizations.h
+++ b/Engine/Plugins/Editor/ModelingToolsEditorMode/Source/ModelingToolsEditorMode/Public/DetailsCustomizations/BakeTransformToolCustomizations.h
@@ -5,6 +5,9 @@
 #include "IDetailCustomization.h"
 #include "Templates/SharedPointer.h"
 
+class IPropertyHandle;
+class FPropertyRestriction;
+
 // Customization for UBakeTransformToolProperties
 class FBakeTransformToolDetails : public IDetailCustomization
 {
