MeshModelingToolset: Fall back to interactive save when creating an asset for an unsaved map in restrictive mode

diff --git a/Engine/Plugins/Editor/ModelingToolsEditorMode/Source/ModelingToolsEditorMode/Private/ModelingModeAssetUtils.cpp b/Engine/Plugins/Editor/ModelingToolsEditorMode/Source/ModelingToolsEditorMode/Private/ModelingModeAssetUtils.cpp
--- a/Engine/Plugins/Editor/ModelingToolsEditorMode/Source/ModelingToolsEditorMode/Private/ModelingModeAssetUtils.cpp
+++ b/Engine/Plugins/Editor/ModelingToolsEditorMode/Source/ModelingToolsEditorMode/Private/ModelingModeAssetUtils.cpp
@@ -89,6 +89,7 @@ FString UE::Modeling::GetNewAssetPathName(const FString& BaseNameIn, const UWorl
 	EModelingModeAssetGenerationBehavior AutoGenMode = Settings->GetAssetGenerationMode();
 
 	FString PackageFolderPath = TEXT("/Game/");
+	bool bFallBackToInteractiveSave = false;
 
 	if (SuggestedFolder.Len() > 0 && SuggestedFolder.StartsWith("/") )
 	{
@@ -106,11 +107,23 @@ FString UE::Modeling::GetNewAssetPathName(const FString& BaseNameIn, const UWorl
 
 		// If the level is unsaved, then GetWorldRelativeAssetRootPath will return "/Temp". This is an undesirable
 		// place to even temporarily save assets because they must be manually moved elsewhere to save, and the /Temp folder 
-		// is somewhat troublesome to get to in the Editor. So if the flag below is set, then we will fall back to /Game
-		// instead. Projects with strict asset polices (due to P4/etc) may not want to allow this by default, so it is configurable.
-		if (PackageFolderPath.StartsWith("/Temp") && Settings->bStoreUnsavedLevelAssetsInTopLevelGameFolder)
+		// is somewhat troublesome to get to in the Editor.
+		if (PackageFolderPath.StartsWith("/Temp"))
 		{
-			PackageFolderPath = TEXT("/Game/");
+			if (Settings->InRestrictiveMode())
+			{
+				// Fall back to the interactively choosing the asset path.
+				bFallBackToInteractiveSave = true;
+			}
+			else
+			{
+				// If the flag below is set, then we will fall back to /Game instead.
+				// Projects with strict asset polices (due to P4/etc) may not want to allow this by default, so it is configurable.
+				if (Settings->bStoreUnsavedLevelAssetsInTopLevelGameFolder)
+				{
+					PackageFolderPath = TEXT("/Game/");
+				}
+			}
 		}
 
 		if (Settings->InRestrictiveMode())
@@ -160,9 +173,10 @@ FString UE::Modeling::GetNewAssetPathName(const FString& BaseNameIn, const UWorl
 
 	// If we are in interactive mode, show the modal dialog and then get the path/name.
 	// If the user cancels, we are going to discard the asset
-	if (AutoGenMode == EModelingModeAssetGenerationBehavior::InteractivePromptToSave)
+	if (AutoGenMode == EModelingModeAssetGenerationBehavior::InteractivePromptToSave || bFallBackToInteractiveSave)
 	{
-		FString SelectedPath = UE::Local::InteractiveSelectAssetPath(ObjectBaseName, LOCTEXT("GenerateStaticMeshActorPathDialogWarning", "Choose Folder Path and Name for New Asset. Cancel to Discard New Asset."));
+		const FString SelectedPath = UE::Local::InteractiveSelectAssetPath(
+			ObjectBaseName, LOCTEXT("GenerateStaticMeshActorPathDialogWarning", "Choose Folder Path and Name for New Asset. Cancel to Discard New Asset."));
 		if (SelectedPath.IsEmpty() == false)
 		{
 			PackageFolderPath = FPaths::GetPath(SelectedPath);
