ModelingMode: remove the ablity to enable/disable DynamicMeshActor support in Modeling Mode. Always enable.

diff --git a/Engine/Plugins/Editor/ModelingToolsEditorMode/Source/ModelingToolsEditorMode/Private/ModelingToolsEditorModeToolkit.cpp b/Engine/Plugins/Editor/ModelingToolsEditorMode/Source/ModelingToolsEditorMode/Private/ModelingToolsEditorModeToolkit.cpp
--- a/Engine/Plugins/Editor/ModelingToolsEditorMode/Source/ModelingToolsEditorMode/Private/ModelingToolsEditorModeToolkit.cpp
+++ b/Engine/Plugins/Editor/ModelingToolsEditorMode/Source/ModelingToolsEditorMode/Private/ModelingToolsEditorModeToolkit.cpp
@@ -1108,10 +1108,10 @@ void FModelingToolsEditorModeToolkit::UpdateObjectCreationOptionsFromSettings()
 	const UModelingToolsEditorModeSettings* Settings = GetDefault<UModelingToolsEditorModeSettings>();
 
 	// enable/disable dynamic mesh actors
-	UCreateMeshObjectTypeProperties::bEnableDynamicMeshActorSupport = Settings->GetEnableDynamicMeshActors();
+	UCreateMeshObjectTypeProperties::bEnableDynamicMeshActorSupport = true;
 
 	// set configured default type
-	if (Settings->DefaultMeshObjectType == EModelingModeDefaultMeshObjectType::DynamicMeshActor && Settings->GetEnableDynamicMeshActors())
+	if (Settings->DefaultMeshObjectType == EModelingModeDefaultMeshObjectType::DynamicMeshActor)
 	{
 		UCreateMeshObjectTypeProperties::DefaultObjectTypeIdentifier = UCreateMeshObjectTypeProperties::DynamicMeshActorIdentifier;
 	}
diff --git a/Engine/Plugins/Editor/ModelingToolsEditorMode/Source/ModelingToolsEditorMode/Public/ModelingToolsEditorModeSettings.h b/Engine/Plugins/Editor/ModelingToolsEditorMode/Source/ModelingToolsEditorMode/Public/ModelingToolsEditorModeSettings.h
--- a/Engine/Plugins/Editor/ModelingToolsEditorMode/Source/ModelingToolsEditorMode/Public/ModelingToolsEditorModeSettings.h
+++ b/Engine/Plugins/Editor/ModelingToolsEditorMode/Source/ModelingToolsEditorMode/Public/ModelingToolsEditorModeSettings.h
@@ -72,7 +72,7 @@ public:
 protected:
 
 	/** Enable/Disable the options to emit Dynamic Mesh Actors in Modeling Mode Tools */
-	UPROPERTY(config, EditAnywhere, Category = "Modeling Mode")
+	UPROPERTY()
 	bool bEnableDynamicMeshActors = false;
 
 	/** Where should Assets auto-generated by Modeling Tools be stored by default */
@@ -85,13 +85,6 @@ protected:
 
 public:
 
-	bool GetEnableDynamicMeshActors() const
-	{
-		return bRestrictiveMode
-					? true
-					: bEnableDynamicMeshActors;
-	}
-
 	EModelingModeAssetGenerationLocation GetAssetGenerationLocation() const
 	{
 		return bRestrictiveMode
