ModelingMode: add support for new element-level geometry selection system for DynamicMeshComponents

New UModelingSelectionInteraction provides click-to-select and gizmo interactions for current set of mesh objects that support selection, by manipulating the active UGeometrySelectionManager.

ModelingMode now creates a UGeometrySelectionManager and UModelingSelectionInteraction on mode startup, if 'persistent selections' are enabled in the project settings. This replaces the previous persistent selection system/manager, which is no longer created/registered.
UModelingToolsEditorMode::ShouldDrawWidget() and ::ComputeBoundingBoxForViewportFocus() updated to support active Selection.

Functions ::OnToolsContextRender(), ::OnToolPostBuild(), ::GetGeometrySelectionChangesAllowed(), ::TestForEditorGizmoHit(), ::UpdateSelectionManagerOnEditorSelectionChange() added to support selection system.

FModelingToolsEditorModeToolkit will add buttons for selection control if Persistent Selections are enabled. Dynamic updates to Palette based on active selection are initiated by EdMode. Currently the selection-actions palette is empty.

diff --git a/Engine/Plugins/Editor/ModelingToolsEditorMode/Source/ModelingToolsEditorMode/ModelingToolsEditorMode.Build.cs b/Engine/Plugins/Editor/ModelingToolsEditorMode/Source/ModelingToolsEditorMode/ModelingToolsEditorMode.Build.cs
--- a/Engine/Plugins/Editor/ModelingToolsEditorMode/Source/ModelingToolsEditorMode/ModelingToolsEditorMode.Build.cs
+++ b/Engine/Plugins/Editor/ModelingToolsEditorMode/Source/ModelingToolsEditorMode/ModelingToolsEditorMode.Build.cs
@@ -46,6 +46,7 @@ public class ModelingToolsEditorMode : ModuleRules
 				"TypedElementRuntime",
 				"InteractiveToolsFramework",
 				"EditorInteractiveToolsFramework",
+				"GeometryFramework",
 				"ModelingComponents",
 				"ModelingComponentsEditorOnly",
 				"MeshModelingTools",
diff --git a/Engine/Plugins/Editor/ModelingToolsEditorMode/Source/ModelingToolsEditorMode/Private/ModelingSelectionInteraction.cpp b/Engine/Plugins/Editor/ModelingToolsEditorMode/Source/ModelingToolsEditorMode/Private/ModelingSelectionInteraction.cpp
new file mode 100644
--- /dev/null
+++ b/Engine/Plugins/Editor/ModelingToolsEditorMode/Source/ModelingToolsEditorMode/Private/ModelingSelectionInteraction.cpp
@@ -0,0 +1,329 @@
+// Copyright Epic Games, Inc. All Rights Reserved.
+
+
+#include "ModelingSelectionInteraction.h"
+#include "InteractiveToolsContext.h"
+#include "BaseGizmos/TransformGizmoUtil.h"
+#include "Selection/GeometrySelectionManager.h"
+#include "ToolSceneQueriesUtil.h"
+#include "SceneQueries/SceneSnappingManager.h"
+
+using namespace UE::Geometry;
+
+void UModelingSelectionInteraction::Initialize(
+	TObjectPtr<UGeometrySelectionManager> SelectionManagerIn,
+	TUniqueFunction<bool()> CanChangeSelectionCallbackIn,
+	TUniqueFunction<bool(const FInputDeviceRay&)> ExternalHitCaptureCallbackIn)
+{
+	SelectionManager = SelectionManagerIn;
+
+	CanChangeSelectionCallback = MoveTemp(CanChangeSelectionCallbackIn);
+	ExternalHitCaptureCallback = MoveTemp(ExternalHitCaptureCallbackIn);
+
+	// create click behavior and set ourselves as click target
+	ClickBehavior = NewObject<USingleClickInputBehavior>();
+	ClickBehavior->Modifiers.RegisterModifier(AddToSelectionModifier, FInputDeviceState::IsShiftKeyDown);
+	ClickBehavior->Modifiers.RegisterModifier(ToggleSelectionModifier, FInputDeviceState::IsCtrlKeyDown);
+	ClickBehavior->Initialize(this);
+
+	BehaviorSet = NewObject<UInputBehaviorSet>();
+	BehaviorSet->Add(ClickBehavior, this);
+
+	TransformProxy = NewObject<UTransformProxy>(this);
+	// todo: make this repositionable etc. Maybe make this function a delegate? or allow caller to provide the gizmo?
+	TransformGizmo = UE::TransformGizmoUtil::Create3AxisTransformGizmo(
+		SelectionManager->GetToolsContext()->GizmoManager, this, TEXT("ModelingSelectionInteraction") );
+	TransformGizmo->SetActiveTarget(TransformProxy, SelectionManager->GetToolsContext()->GizmoManager);
+	TransformGizmo->SetVisibility(false);
+
+	// listen for change events on transform proxy
+	TransformProxy->OnTransformChanged.AddUObject(this, &UModelingSelectionInteraction::OnGizmoTransformChanged);
+	TransformProxy->OnBeginTransformEdit.AddUObject(this, &UModelingSelectionInteraction::OnBeginGizmoTransform);
+	TransformProxy->OnEndTransformEdit.AddUObject(this, &UModelingSelectionInteraction::OnEndGizmoTransform);
+
+	// listen for selection changes to update gizmo
+	OnSelectionModifiedDelegate = SelectionManager->OnSelectionModified.AddUObject(this, &UModelingSelectionInteraction::OnSelectionManager_SelectionModified);
+}
+
+void UModelingSelectionInteraction::Shutdown()
+{
+	if (IsValid(SelectionManager))
+	{
+		SelectionManager->OnSelectionModified.Remove(OnSelectionModifiedDelegate);
+	}
+
+	if (TransformGizmo)
+	{
+		SelectionManager->GetToolsContext()->GizmoManager->DestroyGizmo(TransformGizmo);
+		TransformGizmo = nullptr;
+	}
+}
+
+
+
+
+void UModelingSelectionInteraction::OnUpdateModifierState(int ModifierID, bool bIsOn)
+{
+	// update modifier state flags
+	if (ModifierID == AddToSelectionModifier)
+	{
+		bAddToSelectionEnabled = bIsOn;
+	}
+	else if (ModifierID == ToggleSelectionModifier)
+	{
+		bToggleSelectionEnabled = bIsOn;
+	}
+}
+
+PRAGMA_DISABLE_OPTIMIZATION
+void UModelingSelectionInteraction::ComputeSceneHits(const FInputDeviceRay& ClickPos,
+	bool& bHitActiveObjects, FInputRayHit& ActiveObjectHit,
+	bool& bHitInactiveObjectFirst, FInputRayHit& InactiveObjectHit)
+{
+	ActiveObjectHit = FInputRayHit();
+	bHitActiveObjects = SelectionManager->RayHitTest(ClickPos.WorldRay, ActiveObjectHit);
+
+	// We want to filter out hits against nearer objects. This is...tricky.
+	// TODO: we probably should not actually do this here. Modeling Mode probably needs to 
+	// have a behavior that mimics editor selection, that would have just received the closer
+	// FInputRayHit and taken the click, instead of this function...
+	bHitInactiveObjectFirst = false;
+	InactiveObjectHit = FInputRayHit();
+	if (bHitActiveObjects)
+	{
+		TArray<const UPrimitiveComponent*> IgnoreComponents;
+		if (ActiveObjectHit.HitObject.IsValid())
+		{
+			if (UPrimitiveComponent* SelectionComponent = Cast<UPrimitiveComponent>(ActiveObjectHit.HitObject.Get()) )
+			{
+				IgnoreComponents.Add(SelectionComponent);
+			}
+		}
+		FHitResult HitResultOut;
+		if (ToolSceneQueriesUtil::FindNearestVisibleObjectHit(USceneSnappingManager::Find(SelectionManager->GetToolsContext()->ToolManager), 
+			HitResultOut, ClickPos.WorldRay, &IgnoreComponents, nullptr))
+		{
+			if (HitResultOut.Distance < ActiveObjectHit.HitDepth)
+			{
+				bHitInactiveObjectFirst = true;
+				InactiveObjectHit.HitDepth = HitResultOut.Distance;
+				InactiveObjectHit.bHit = true;
+			}
+		}
+	}
+}
+PRAGMA_ENABLE_OPTIMIZATION
+
+PRAGMA_DISABLE_OPTIMIZATION
+FInputRayHit UModelingSelectionInteraction::IsHitByClick(const FInputDeviceRay& ClickPos)
+{
+	// ignore hits in these cases
+	if (  (SelectionManager->GetMeshTopologyMode() == UGeometrySelectionManager::EMeshTopologyMode::None) 
+		|| (CanChangeSelectionCallback() == false)
+		|| ExternalHitCaptureCallback(ClickPos) )
+	{
+		return FInputRayHit();
+	}
+
+	bool bHitActiveObjects, bHitInactiveObjectFirst;
+	FInputRayHit ActiveObjectHit, InactiveObjectHit;
+	ComputeSceneHits(ClickPos, bHitActiveObjects, ActiveObjectHit, bHitInactiveObjectFirst, InactiveObjectHit);
+
+	// Note: currently this flow will produce somewhat undesirable behavior, in that if 
+	// some other object not in the selection target(s) is clicked, the selection will be 
+	// cleared but that object will have to be clicked again to select it. We cannot currently 
+	// clear the mesh-selection and change the actor-selection inside a single transaction in a
+	// single click. This is because we are relying on the Editor to do the actor-selection change.
+	// So if we were to (eg) ClearSelection() and return no-hit here, the ClearSelection() call
+	// would emit a transaction, and then the actor-selection-change would emit a second transaction.
+	// 
+	// Behavior in many DCCS is that they can (1) clear active mesh-selection, (2) select new target
+	// object, and (3) do new mesh-selection in a single click. This will require quite a bit of effort
+	// to achieve in the future. 
+	//
+	// So for now we set a bClearSelectionOnClicked flag and consume the click. This means the 1-2-3
+	// flow above requires 3 separate clicks :( 
+
+
+	bClearSelectionOnClicked = false;
+	if (SelectionManager->HasSelection())
+	{
+		if (bHitInactiveObjectFirst)
+		{
+			//bClearSelectionOnClicked = true;
+			//return InactiveObjectHit;
+			return FInputRayHit();
+		}
+		else if (bHitActiveObjects)
+		{
+			return ActiveObjectHit;
+		}
+		else
+		{
+			// if we have active selection we still want to capture this hit so that we can clear in OnClicked
+			//bClearSelectionOnClicked = true;
+			//return FInputRayHit(TNumericLimits<double>::Max());
+			return FInputRayHit();
+		}
+	}
+	else
+	{
+		if (bHitActiveObjects && bHitInactiveObjectFirst == false)
+		{
+			return ActiveObjectHit;
+		}
+	}
+
+	return FInputRayHit();
+}
+PRAGMA_ENABLE_OPTIMIZATION
+
+
+
+void UModelingSelectionInteraction::OnClicked(const FInputDeviceRay& ClickPos)
+{
+	// this flag is set in IsHitByClick test
+	if (bClearSelectionOnClicked)
+	{
+		SelectionManager->ClearSelection();
+		bClearSelectionOnClicked = false;
+		return;
+	}
+
+	FGeometrySelectionUpdateConfig UpdateConfig;
+	UpdateConfig.ChangeType = EGeometrySelectionChangeType::Replace;
+	if (bAddToSelectionEnabled)
+	{
+		UpdateConfig.ChangeType = EGeometrySelectionChangeType::Add;
+	}
+	else if (bToggleSelectionEnabled)
+	{
+		UpdateConfig.ChangeType = EGeometrySelectionChangeType::Remove;
+	}
+
+	FGeometrySelectionUpdateResult Result;
+	SelectionManager->UpdateSelectionViaRaycast(
+		ClickPos.WorldRay,
+		UpdateConfig,
+		Result);
+}
+
+
+void UModelingSelectionInteraction::UpdateGizmoOnSelectionChange()
+{
+	if (SelectionManager->HasSelection() == false)
+	{
+		TransformGizmo->SetVisibility(false);
+	}
+	else
+	{
+		TransformGizmo->SetVisibility(true);
+
+		FFrame3d SelectionFrame;
+		SelectionManager->GetSelectionWorldFrame(SelectionFrame);
+		TransformGizmo->ReinitializeGizmoTransform( SelectionFrame.ToFTransform() );
+
+		//FGeometrySelectionBounds Bounds;
+		//SelectionManager->GetSelectionBounds(Bounds);
+		//TransformGizmo->ReinitializeGizmoTransform( FTransform(Bounds.WorldBounds.Center()) );
+	}
+}
+
+
+void UModelingSelectionInteraction::OnSelectionManager_SelectionModified()
+{
+	UpdateGizmoOnSelectionChange();
+}
+
+
+
+void UModelingSelectionInteraction::OnBeginGizmoTransform(UTransformProxy* Proxy)
+{
+	FTransform Transform = Proxy->GetTransform();
+	InitialGizmoFrame = FFrame3d(Transform);
+	InitialGizmoScale = FVector3d(Transform.GetScale3D());
+
+	bInActiveTransform = SelectionManager->BeginTransformation();
+
+	OnTransformBegin.Broadcast();
+}
+
+void UModelingSelectionInteraction::OnEndGizmoTransform(UTransformProxy* Proxy)
+{
+	if (bInActiveTransform)
+	{
+		SelectionManager->EndTransformation();
+		bInActiveTransform = false;
+
+		OnTransformEnd.Broadcast();
+	}
+}
+
+void UModelingSelectionInteraction::OnGizmoTransformChanged(UTransformProxy* Proxy, FTransform Transform)
+{
+	if (bInActiveTransform)
+	{
+		LastUpdateGizmoFrame = FFrame3d(Transform);
+		LastUpdateGizmoScale = FVector3d(Transform.GetScale3D());
+		bGizmoUpdatePending = true;
+	}
+}
+
+
+void UModelingSelectionInteraction::ApplyPendingTransformInteractions()
+{
+	if (bInActiveTransform == false || bGizmoUpdatePending == false)
+	{
+		return;
+	}
+
+	FFrame3d CurFrame = LastUpdateGizmoFrame;
+	FVector3d CurScale = LastUpdateGizmoScale;
+	FVector3d TranslationDelta = CurFrame.Origin - InitialGizmoFrame.Origin;
+	FQuaterniond RotateDelta = CurFrame.Rotation - InitialGizmoFrame.Rotation;
+	FVector3d CurScaleDelta = CurScale - InitialGizmoScale;
+
+	const bool bLastUpdateUsedWorldFrame = true;	// for later local-frame support
+
+	if (TranslationDelta.SquaredLength() > 0.0001 || RotateDelta.SquaredLength() > 0.0001 || CurScaleDelta.SquaredLength() > 0.0001)
+	{
+		if (bLastUpdateUsedWorldFrame)
+		{
+			// For a world frame gizmo, the scaling needs to happen in world aligned gizmo space, but the 
+			// rotation is still encoded in the local gizmo frame change.
+			FQuaterniond RotationToApply = CurFrame.Rotation * InitialGizmoFrame.Rotation.Inverse();
+			SelectionManager->UpdateTransformation( [&](int32 VertexID, const FVector3d& InitialPosition, const FTransform3d& WorldTransform)
+			{
+				FVector3d LocalTranslation = WorldTransform.InverseTransformVector(TranslationDelta);
+
+				FVector3d PosLocal = InitialPosition;
+				FVector3d PosWorld = WorldTransform.TransformPosition(PosLocal);
+				FVector3d PosWorldGizmo = PosWorld - InitialGizmoFrame.Origin;
+
+				FVector3d NewPosWorld = RotationToApply * (PosWorldGizmo * CurScale) + CurFrame.Origin;
+				FVector3d NewPosLocal = WorldTransform.InverseTransformPosition(NewPosWorld);
+				return NewPosLocal;
+			});
+		}
+		else
+		{
+			SelectionManager->UpdateTransformation( [&](int32 VertexID, const FVector3d& InitialPosition, const FTransform3d& WorldTransform)
+			{
+				// For a local gizmo, we just get the coordinates in the original frame, scale in that frame,
+				// then interpret them as coordinates in the new frame.
+				FVector3d PosLocal = InitialPosition;
+				FVector3d PosWorld = WorldTransform.TransformPosition(PosLocal);
+				FVector3d PosGizmo = InitialGizmoFrame.ToFramePoint(PosWorld);
+				PosGizmo = CurScale * PosGizmo;
+				FVector3d NewPosWorld = CurFrame.FromFramePoint(PosGizmo);
+				FVector3d NewPosLocal = WorldTransform.InverseTransformPosition(NewPosWorld);
+				return NewPosLocal;
+			});
+		}
+	}
+
+
+
+
+	bGizmoUpdatePending = false;
+}
\ No newline at end of file
diff --git a/Engine/Plugins/Editor/ModelingToolsEditorMode/Source/ModelingToolsEditorMode/Private/ModelingToolsEditorMode.cpp b/Engine/Plugins/Editor/ModelingToolsEditorMode/Source/ModelingToolsEditorMode/Private/ModelingToolsEditorMode.cpp
--- a/Engine/Plugins/Editor/ModelingToolsEditorMode/Source/ModelingToolsEditorMode/Private/ModelingToolsEditorMode.cpp
+++ b/Engine/Plugins/Editor/ModelingToolsEditorMode/Source/ModelingToolsEditorMode/Private/ModelingToolsEditorMode.cpp
@@ -17,7 +17,6 @@
 #include "EditorViewportClient.h"
 #include "EngineAnalytics.h"
 #include "BaseGizmos/TransformGizmoUtil.h"
-#include "Selection/PersistentMeshSelectionManager.h"
 #include "Selection/StoredMeshSelectionUtil.h"
 #include "Snapping/ModelingSceneSnappingManager.h"
 #include "Scene/LevelObjectsObserver.h"
@@ -102,6 +101,7 @@
 #include "Tools/LODManagerTool.h"
 
 #include "EditorModeManager.h"
+#include "UnrealWidget.h"
 
 // stylus support
 #include "IStylusInputModule.h"
@@ -115,6 +115,11 @@
 #include "ModelingModeAssetUtils.h"
 #include "EditorModelingObjectsCreationAPI.h"
 
+#include "Selection/GeometrySelectionManager.h"
+#include "Selection/DynamicMeshSelector.h"
+#include "ModelingSelectionInteraction.h"
+#include "DynamicMeshActor.h"
+
 #define LOCTEXT_NAMESPACE "UModelingToolsEditorMode"
 
 
@@ -204,12 +209,18 @@ bool UModelingToolsEditorMode::CanAutoSave() const
 
 bool UModelingToolsEditorMode::ShouldDrawWidget() const
 { 
-	// allow standard xform gizmo if we don't have an active tool
+	// hide standard xform gizmo if we have an active tool
 	if (GetInteractiveToolsContext() != nullptr && GetToolManager()->HasAnyActiveTool())
 	{
 		return false;
 	}
 
+	// hide standard xform gizmo if we have an active selection
+	if (GetSelectionManager() && GetSelectionManager()->HasSelection())
+	{
+		return false;
+	}
+
 	return UBaseLegacyWidgetEdMode::ShouldDrawWidget(); 
 }
 
@@ -325,6 +336,16 @@ void UModelingToolsEditorMode::Enter()
 	// skin weights.
 	GetInteractiveToolsContext()->TargetManager->AddTargetFactory(NewObject<USkeletalMeshComponentReadOnlyToolTargetFactory>(GetToolManager()));
 
+	// listen to post-build
+	GetToolManager()->OnToolPostBuild.AddUObject(this, &UModelingToolsEditorMode::OnToolPostBuild);
+
+	// forward shutdown requests
+	GetToolManager()->OnToolShutdownRequest.BindLambda([this](UInteractiveToolManager*, UInteractiveTool* Tool, EToolShutdownType ShutdownType)
+	{
+		GetInteractiveToolsContext()->EndTool(ShutdownType); 
+		return true;
+	});
+
 	// register stylus event handler
 	StylusStateTracker = MakeUnique<FStylusStateTracker>();
 
@@ -335,6 +356,40 @@ void UModelingToolsEditorMode::Enter()
 	UE::Geometry::RegisterSceneSnappingManager(GetInteractiveToolsContext());
 	SceneSnappingManager = UE::Geometry::FindModelingSceneSnappingManager(GetToolManager());
 
+	// register selection manager, if this feature is enabled in the mode settings
+	const UModelingToolsEditorModeSettings* ModelingModeSettings = GetDefault<UModelingToolsEditorModeSettings>();
+	if (ModelingModeSettings && ModelingModeSettings->bEnablePersistentSelections)
+	{
+		// set up SelectionManager and register known factory types
+		SelectionManager = NewObject<UGeometrySelectionManager>(GetToolManager());
+		SelectionManager->Initialize(GetInteractiveToolsContext(), GetToolManager()->GetContextTransactionsAPI());
+		SelectionManager->RegisterSelectorFactory(MakeUnique<FDynamicMeshComponentSelectorFactory>());
+
+		GetInteractiveToolsContext()->OnRender.AddUObject(this, &UModelingToolsEditorMode::OnToolsContextRender);
+		// this is hopefully temporary? kinda gross...
+		GetInteractiveToolsContext()->ContextObjectStore->AddContextObject(SelectionManager);
+
+		// rebuild tool palette on any selection changes. This is expensive and ideally will be
+		// optimized in the future.
+		SelectionManager_SelectionModifiedHandle = SelectionManager->OnSelectionModified.AddLambda([this]()
+		{
+			((FModelingToolsEditorModeToolkit*)Toolkit.Get())->ForceToolPaletteRebuild();
+		});
+
+		// set up the selection interaction
+		SelectionInteraction = NewObject<UModelingSelectionInteraction>(GetToolManager());
+		SelectionInteraction->Initialize(SelectionManager,
+			[this]() { return GetGeometrySelectionChangesAllowed(); },
+			[this](const FInputDeviceRay& DeviceRay) { return TestForEditorGizmoHit(DeviceRay); });
+		GetInteractiveToolsContext()->InputRouter->RegisterSource(SelectionInteraction);
+
+		// Disable the SnappingManager while the SelectionInteraction is editing a mesh via transform gizmo
+		SelectionInteraction->OnTransformBegin.AddLambda(
+			[this]() { SceneSnappingManager->PauseSceneGeometryUpdates(); });
+		SelectionInteraction->OnTransformEnd.AddLambda(
+			[this]() { SceneSnappingManager->UnPauseSceneGeometryUpdates(); });
+	}
+
 	// register level objects observer that will update the snapping manager as the scene changes
 	LevelObjectsObserver = MakeShared<FLevelObjectsObserver>();
 	LevelObjectsObserver->OnActorAdded.AddLambda([this](AActor* Actor)
@@ -354,13 +409,6 @@ void UModelingToolsEditorMode::Enter()
 	// tracker will auto-populate w/ the current level, but must have registered the handlers first!
 	LevelObjectsObserver->Initialize(GetWorld());
 
-	// register selection manager, if this feature is enabled in the mode settings
-	const UModelingToolsEditorModeSettings* ModelingModeSettings = GetDefault<UModelingToolsEditorModeSettings>();
-	if (ModelingModeSettings && ModelingModeSettings->bEnablePersistentSelections)
-	{
-		UE::Geometry::RegisterPersistentMeshSelectionManager(GetInteractiveToolsContext());
-	}
-
 	// disable HitProxy rendering, it is not used in Modeling Mode and adds overhead to Render() calls
 	GetInteractiveToolsContext()->SetEnableRenderingDuringHitProxyPass(false);
 
@@ -650,6 +698,38 @@ void UModelingToolsEditorMode::Enter()
 	RegisterPolyModelActionTool(EEditMeshPolygonsToolActions::CutFaces, ToolManagerCommands.BeginPolyModelTool_CutFaces, TEXT("PolyEdit_CutFaces"));
 
 
+	// set up selection type toggles
+	auto RegisterSelectionTopologyType = [this](UGeometrySelectionManager::EMeshTopologyMode TopoMode, TSharedPtr<FUICommandInfo> UICommand)
+	{
+		Toolkit->GetToolkitCommands()->MapAction(UICommand,
+			FExecuteAction::CreateLambda([this, TopoMode]() { GetSelectionManager()->SetMeshTopologyMode(TopoMode); }),
+			FCanExecuteAction(),
+			FIsActionChecked::CreateLambda([this, TopoMode]() { return GetSelectionManager()->GetMeshTopologyMode() == TopoMode; }),
+			EUIActionRepeatMode::RepeatDisabled);
+	};
+	if (GetSelectionManager() != nullptr)
+	{
+		RegisterSelectionTopologyType(UGeometrySelectionManager::EMeshTopologyMode::None, ToolManagerCommands.BeginSelectionAction_ToObjectType);
+		RegisterSelectionTopologyType(UGeometrySelectionManager::EMeshTopologyMode::Triangle, ToolManagerCommands.BeginSelectionAction_ToTriangleType);
+		RegisterSelectionTopologyType(UGeometrySelectionManager::EMeshTopologyMode::Polygroup, ToolManagerCommands.BeginSelectionAction_ToPolygroupType);
+	}
+
+	auto RegisterSelectionElementType = [this](EGeometryElementType ElementMode, TSharedPtr<FUICommandInfo> UICommand)
+	{
+		Toolkit->GetToolkitCommands()->MapAction(UICommand,
+			FExecuteAction::CreateLambda([this, ElementMode]() { GetSelectionManager()->SetSelectionElementType(ElementMode); }),
+			FCanExecuteAction(),
+			FIsActionChecked::CreateLambda([this, ElementMode]() { return GetSelectionManager()->GetSelectionElementType() == ElementMode; }),
+			EUIActionRepeatMode::RepeatDisabled);
+	};
+	if (GetSelectionManager() != nullptr)
+	{
+		RegisterSelectionElementType(EGeometryElementType::Vertex, ToolManagerCommands.BeginSelectionAction_ToVertexType);
+		RegisterSelectionElementType(EGeometryElementType::Edge, ToolManagerCommands.BeginSelectionAction_ToEdgeType);
+		RegisterSelectionElementType(EGeometryElementType::Face, ToolManagerCommands.BeginSelectionAction_ToFaceType);
+	}
+
+
 	// register extensions
 	TArray<IModelingModeToolExtension*> Extensions = IModularFeatures::Get().GetModularFeatureImplementations<IModelingModeToolExtension>(
 		IModelingModeToolExtension::GetModularFeatureName());
@@ -743,35 +823,31 @@ void UModelingToolsEditorMode::Enter()
 		ModelingToolkit->InitializeAfterModeSetup();
 	}
 
-	// Need to know about selection changes to clear mesh selections, however we do not want to clear the
-	// mesh selection after a selection change due to transactions, as this may clear a selection we just restored.
-	// Unfortunately most routes to find out about selection changes don't allow for this. Currently this
-	// OnPreChange event is the only one that appears to provide the desired behavior, however it is likely
-	// that this is not going to be reliable...
-	if (ModelingModeSettings && ModelingModeSettings->bEnablePersistentSelections)
-	{
-		SelectionModifiedEventHandle = GetModeManager()->GetSelectedActors()->GetElementSelectionSet()->OnPreChange().AddLambda([&](const UTypedElementSelectionSet*)
-		{
-			if (GIsTransacting == 0)
-			{
-				UE::Geometry::ClearActiveToolSelection(GetToolManager());
-			}
-		});
-	}
+	// Need to know about selection changes to (eg) clear mesh selections. 
+	// Listening to USelection::SelectionChangedEvent here instead of the underlying UTypedElementSelectionSet
+	// events because they do not fire at the right times, particular wrt undo/redo. 
+	SelectionModifiedEventHandle = GetModeManager()->GetSelectedActors()->SelectionChangedEvent.AddLambda(
+		[this](const UObject* Object)  { UpdateSelectionManagerOnEditorSelectionChange(); } );
+
 }
 
+
+
 void UModelingToolsEditorMode::Exit()
 {
-	// clear any active selection
-	UE::Geometry::ClearActiveToolSelection(GetToolManager());
-
-	// deregister selection manager (note: may not have been registered, depending on mode settings)
-	UE::Geometry::DeregisterPersistentMeshSelectionManager(GetInteractiveToolsContext());
+	// shutdown selection interaction
+	if (SelectionInteraction != nullptr)
+	{
+		SelectionInteraction->Shutdown();
+		GetInteractiveToolsContext()->InputRouter->DeregisterSource(SelectionInteraction);
+		SelectionInteraction = nullptr;
+	}
 
 	// stop listening to selection changes. On Editor Shutdown, some of these values become null, which will result in an ensure/crash
-	if (SelectionModifiedEventHandle.IsValid() && UObjectInitialized() && GetModeManager() && GetModeManager()->GetSelectedActors() && GetModeManager()->GetSelectedActors()->GetElementSelectionSet() )
+	if (SelectionModifiedEventHandle.IsValid() && UObjectInitialized() && GetModeManager() 
+		&& ( GetModeManager()->GetSelectedActors() != nullptr) )
 	{
-		GetModeManager()->GetSelectedActors()->GetElementSelectionSet()->OnPreChange().Remove(SelectionModifiedEventHandle);
+		GetModeManager()->GetSelectedActors()->SelectionChangedEvent.Remove(SelectionModifiedEventHandle);
 	}
 
 	// exit any exclusive active tools w/ cancel
@@ -783,6 +859,21 @@ void UModelingToolsEditorMode::Exit()
 		}
 	}
 
+	// Shutdown SelectionManager. Wait until after Tool shutdown in case some restore-selection
+	// is involved (although since we are exiting Mode this currently would never matter)
+	if (SelectionManager != nullptr)
+	{
+		SelectionManager->OnSelectionModified.Remove(SelectionManager_SelectionModifiedHandle);
+		SelectionManager->ClearSelection();
+		SelectionManager->ClearActiveTargets();
+		SelectionManager->Shutdown();
+
+		// hopefully temporary...remove SelectionManager from ContextObjectStore
+		GetInteractiveToolsContext()->ContextObjectStore->RemoveContextObject(SelectionManager);
+
+		SelectionManager = nullptr;
+	}
+
 	//
 	// Engine Analytics
 	//
@@ -833,6 +924,22 @@ void UModelingToolsEditorMode::Exit()
 	UEdMode::Exit();
 }
 
+void UModelingToolsEditorMode::OnToolsContextRender(IToolsContextRenderAPI* RenderAPI)
+{
+	if (SelectionInteraction)
+	{
+		// Bake in transform changes. Note that if we do this in OnToolsContextTick, it will 
+		// still block rendering updates if it is too expensive, unless it is only done every second Tick
+		SelectionInteraction->ApplyPendingTransformInteractions();
+	}
+
+	if (SelectionManager)
+	{
+		// currently relying on debug rendering to visualize selections
+		SelectionManager->DebugRender(RenderAPI);
+	}
+}
+
 bool UModelingToolsEditorMode::ShouldToolStartBeAllowed(const FString& ToolIdentifier) const
 {
 	if (UInteractiveToolManager* Manager = GetToolManager())
@@ -850,11 +957,90 @@ bool UModelingToolsEditorMode::ShouldToolStartBeAllowed(const FString& ToolIdent
 }
 
 
+bool UModelingToolsEditorMode::GetGeometrySelectionChangesAllowed() const
+{
+	// disable selection system if we are in a Tool
+	if ( GetToolManager() && GetToolManager()->HasAnyActiveTool() )
+	{
+		return false;
+	}
+	return true;
+}
+
+bool UModelingToolsEditorMode::TestForEditorGizmoHit(const FInputDeviceRay& ClickPos) const
+{
+	// Because the editor gizmo does not participate in InputRouter behavior system, in some input behaviors
+	// we need to filter out clicks on the gizmo. This function can do this check.
+	if ( ShouldDrawWidget() )
+	{
+		FViewport* FocusedViewport = GetToolManager()->GetContextQueriesAPI()->GetFocusedViewport();
+		HHitProxy* HitResult = FocusedViewport->GetHitProxy(ClickPos.ScreenPosition.X, ClickPos.ScreenPosition.Y);
+		if (HitResult && HitResult->IsA(HWidgetAxis::StaticGetType()))
+		{
+			return true;
+		}
+	}
+	return false;
+}
+
+
+
+void UModelingToolsEditorMode::UpdateSelectionManagerOnEditorSelectionChange()
+{
+	if (!SelectionManager) return;
+
+	// if we are in undo/redo, ignore selection change notifications, the required
+	// changes are handled via FChanges that SelectionManager has emitted
+	if (GIsTransacting)
+	{
+		return;
+	}
+
+	// If Editor is creating a transaction, we assume we must be in a selection change.
+	// Need to handle all SelectionManager changes (deselect + change-targets) during
+	// the transaction so that undo/redo works properly.
+	bool bCreatingTransaction = (GUndo != nullptr);
+	if (bCreatingTransaction)
+	{
+		TArray<AActor*> SelectedActors;
+		GEditor->GetSelectedActors()->GetSelectedObjects<AActor>(SelectedActors); 
+
+		TArray<FGeometryIdentifier> ValidIdentifiers;
+		for ( AActor* Actor : SelectedActors)
+		{
+			if (Cast<ADynamicMeshActor>(Actor) != nullptr)
+			{
+				ValidIdentifiers.Add( FGeometryIdentifier::PrimitiveComponent(
+					Cast<ADynamicMeshActor>(Actor)->GetDynamicMeshComponent(), FGeometryIdentifier::EObjectType::DynamicMeshComponent) );
+			}
+		}
+
+		SelectionManager->SynchronizeActiveTargets(ValidIdentifiers,
+			[this]() {
+				SelectionManager->ClearSelection();
+			});
+	}
+
+}
+
+
 void UModelingToolsEditorMode::CreateToolkit()
 {
 	Toolkit = MakeShareable(new FModelingToolsEditorModeToolkit);
 }
 
+void UModelingToolsEditorMode::OnToolPostBuild(
+	UInteractiveToolManager* InToolManager, EToolSide InSide, 
+	UInteractiveTool* InBuiltTool, UInteractiveToolBuilder* InToolBuilder, const FToolBuilderState& ToolState)
+{
+	// Want to clear active selection when a Tool starts, but we have to wait until after it has been
+	// built, so that the Tool has a chance to see the Selection
+	if (GetSelectionManager() && GetSelectionManager()->HasSelection())
+	{
+		GetSelectionManager()->ClearSelection();
+	}
+}
+
 void UModelingToolsEditorMode::OnToolStarted(UInteractiveToolManager* Manager, UInteractiveTool* Tool)
 {
 	FModelingToolActionCommands::UpdateToolCommandBinding(Tool, Toolkit->GetToolkitCommands(), false);
@@ -1037,6 +1223,14 @@ void UModelingToolsEditorMode::FocusCameraAtCursorHotkey()
 
 bool UModelingToolsEditorMode::ComputeBoundingBoxForViewportFocus(AActor* Actor, UPrimitiveComponent* PrimitiveComponent, FBox& InOutBox) const
 {
+	auto ProcessFocusBoxFunc = [](FBox& FocusBoxInOut)
+	{
+		double MaxDimension = FocusBoxInOut.GetExtent().GetMax();
+		double ExpandAmount = (MaxDimension > SMALL_NUMBER) ? (MaxDimension * 0.2) : 25;		// 25 is a bit arbitrary here...
+		FocusBoxInOut = FocusBoxInOut.ExpandBy(MaxDimension * 0.2);
+	};
+
+	// if Tool supports custom Focus box, use that
 	if (GetToolManager()->HasAnyActiveTool())
 	{
 		UInteractiveTool* Tool = GetToolManager()->GetActiveTool(EToolSide::Mouse);
@@ -1046,19 +1240,23 @@ bool UModelingToolsEditorMode::ComputeBoundingBoxForViewportFocus(AActor* Actor,
 			InOutBox = FocusAPI->GetWorldSpaceFocusBox();
 			if (InOutBox.IsValid)
 			{
-				float MaxDimension = InOutBox.GetExtent().GetMax();
-				if (MaxDimension > SMALL_NUMBER)
-				{
-					InOutBox = InOutBox.ExpandBy(MaxDimension * 0.2f);
-				}
-				else
-				{
-					InOutBox = InOutBox.ExpandBy(25);
-				}
+				ProcessFocusBoxFunc(InOutBox);
 				return true;
 			}
 		}
 	}
+
+	// if we have an active Selection we can focus on that
+	if (GetSelectionManager() && GetSelectionManager()->HasSelection())
+	{
+		FGeometrySelectionBounds SelectionBounds;
+		GetSelectionManager()->GetSelectionBounds(SelectionBounds);
+		InOutBox = (FBox)SelectionBounds.WorldBounds;
+		ProcessFocusBoxFunc(InOutBox);
+		return true;
+	}
+
+	// fallback to base focus behavior
 	return false;
 }
 
diff --git a/Engine/Plugins/Editor/ModelingToolsEditorMode/Source/ModelingToolsEditorMode/Private/ModelingToolsEditorModeToolkit.cpp b/Engine/Plugins/Editor/ModelingToolsEditorMode/Source/ModelingToolsEditorMode/Private/ModelingToolsEditorModeToolkit.cpp
--- a/Engine/Plugins/Editor/ModelingToolsEditorMode/Source/ModelingToolsEditorMode/Private/ModelingToolsEditorModeToolkit.cpp
+++ b/Engine/Plugins/Editor/ModelingToolsEditorMode/Source/ModelingToolsEditorMode/Private/ModelingToolsEditorModeToolkit.cpp
@@ -5,6 +5,7 @@
 #include "ModelingToolsManagerActions.h"
 #include "ModelingToolsEditorModeSettings.h"
 #include "ModelingToolsEditorModeStyle.h"
+#include "Selection/GeometrySelectionManager.h"
 #include "Engine/Selection.h"
 
 #include "Framework/MultiBox/MultiBoxBuilder.h"
@@ -488,6 +489,9 @@ static const FName LODToolsTabName(TEXT("LODs"));
 static const FName BakingToolsTabName(TEXT("Baking"));
 static const FName ModelingFavoritesTabName(TEXT("Favorites"));
 
+static const FName SelectionModesTabName(TEXT("SelectionModes"));
+static const FName SelectionActionsTabName(TEXT("SelectionActions"));
+
 
 const TArray<FName> FModelingToolsEditorModeToolkit::PaletteNames_Standard = { PrimitiveTabName, CreateTabName, PolyModelingTabName, TriModelingTabName, DeformTabName, TransformTabName, MeshProcessingTabName, VoxToolsTabName, AttributesTabName, UVTabName, BakingToolsTabName, VolumesTabName, LODToolsTabName };
 
@@ -502,6 +506,18 @@ void FModelingToolsEditorModeToolkit::GetToolPaletteNames(TArray<FName>& Palette
 		ExistingNames.Add(Name);
 	}
 
+	const UModelingToolsEditorModeSettings* ModelingModeSettings = GetDefault<UModelingToolsEditorModeSettings>();
+	bool bEnableSelectionUI = ModelingModeSettings && ModelingModeSettings->bEnablePersistentSelections;
+	if (bEnableSelectionUI)
+	{
+		if (bShowActiveSelectionActions)
+		{
+			PaletteNames.Insert(SelectionActionsTabName, 0);
+			ExistingNames.Add(SelectionActionsTabName);
+		}
+
+		PaletteNames.Insert(SelectionModesTabName, 0);
+	}
 
 	bool bEnablePrototypes = (CVarEnablePrototypeModelingTools.GetValueOnGameThread() > 0);
 	if (bEnablePrototypes)
@@ -597,7 +613,22 @@ void FModelingToolsEditorModeToolkit::BuildToolPalette(FName PaletteIndex, class
 			}
 		}
 	}
-	if (PaletteIndex == PrimitiveTabName)
+	else if (PaletteIndex == SelectionModesTabName)
+	{
+		ToolbarBuilder.AddToolBarButton(Commands.BeginSelectionAction_ToObjectType);
+		ToolbarBuilder.AddToolBarButton(Commands.BeginSelectionAction_ToTriangleType);
+		ToolbarBuilder.AddToolBarButton(Commands.BeginSelectionAction_ToPolygroupType);
+
+		ToolbarBuilder.AddToolBarButton(Commands.BeginSelectionAction_ToVertexType);
+		ToolbarBuilder.AddToolBarButton(Commands.BeginSelectionAction_ToEdgeType);
+		ToolbarBuilder.AddToolBarButton(Commands.BeginSelectionAction_ToFaceType);
+	}
+	else if (PaletteIndex == SelectionActionsTabName)
+	{
+		//ToolbarBuilder.AddToolBarButton(Commands.BeginSelectionAction_Delete);
+		//ToolbarBuilder.AddToolBarButton(Commands.BeginSelectionAction_Extrude);
+	}
+	else if (PaletteIndex == PrimitiveTabName)
 	{
 		ToolbarBuilder.AddToolBarButton(Commands.BeginAddBoxPrimitiveTool);
 		ToolbarBuilder.AddToolBarButton(Commands.BeginAddSpherePrimitiveTool);
@@ -923,6 +954,27 @@ void FModelingToolsEditorModeToolkit::InvokeUI()
 
 
 
+void FModelingToolsEditorModeToolkit::ForceToolPaletteRebuild()
+{
+
+	UGeometrySelectionManager* SelectionManager = Cast<UModelingToolsEditorMode>(GetScriptableEditorMode())->GetSelectionManager();
+	if (SelectionManager)
+	{
+		bool bHasActiveSelection = SelectionManager->HasSelection();
+
+		if (bShowActiveSelectionActions != bHasActiveSelection)
+		{
+			bShowActiveSelectionActions = bHasActiveSelection;
+			this->RebuildModeToolPalette();
+		}
+	}
+	else
+	{
+		bShowActiveSelectionActions = false;
+	}
+}
+
+
 void FModelingToolsEditorModeToolkit::OnToolPaletteChanged(FName PaletteName) 
 {
 }
diff --git a/Engine/Plugins/Editor/ModelingToolsEditorMode/Source/ModelingToolsEditorMode/Private/ModelingToolsManagerActions.cpp b/Engine/Plugins/Editor/ModelingToolsEditorMode/Source/ModelingToolsEditorMode/Private/ModelingToolsManagerActions.cpp
--- a/Engine/Plugins/Editor/ModelingToolsEditorMode/Source/ModelingToolsEditorMode/Private/ModelingToolsManagerActions.cpp
+++ b/Engine/Plugins/Editor/ModelingToolsEditorMode/Source/ModelingToolsEditorMode/Private/ModelingToolsManagerActions.cpp
@@ -172,6 +172,16 @@ void FModelingToolsManagerCommands::RegisterCommands()
 	REGISTER_MODELING_TOOL_COMMAND(BeginPolyModelTool_Outset, "Outset", "PolyGroup Outset Tool");
 	REGISTER_MODELING_TOOL_COMMAND(BeginPolyModelTool_CutFaces, "Cut", "PolyGroup Cut Faces Tool");
 
+
+	REGISTER_MODELING_TOOL_COMMAND(BeginSelectionAction_ToObjectType, "None", "Disable Geometry Selection");
+	REGISTER_MODELING_TOOL_COMMAND(BeginSelectionAction_ToTriangleType, "Tris", "Set Selection Layer to Triangles");
+	REGISTER_MODELING_TOOL_COMMAND(BeginSelectionAction_ToPolygroupType, "Groups", "Set Selection Layer to Polygroups");
+
+	REGISTER_MODELING_TOOL_COMMAND(BeginSelectionAction_ToVertexType, "Vertex", "Set Selection Type to Vertices");
+	REGISTER_MODELING_TOOL_COMMAND(BeginSelectionAction_ToEdgeType, "Edge", "Set Selection Type to Edges");
+	REGISTER_MODELING_TOOL_COMMAND(BeginSelectionAction_ToFaceType, "Face", "Set Selection Type to Faces");
+
+
 	UI_COMMAND(AcceptActiveTool, "Accept", "Accept the active tool", EUserInterfaceActionType::Button, FInputChord());
 	UI_COMMAND(CancelActiveTool, "Cancel", "Cancel the active tool", EUserInterfaceActionType::Button, FInputChord());
 	UI_COMMAND(CompleteActiveTool, "Done", "Complete the active tool", EUserInterfaceActionType::Button, FInputChord());
diff --git a/Engine/Plugins/Editor/ModelingToolsEditorMode/Source/ModelingToolsEditorMode/Public/ModelingSelectionInteraction.h b/Engine/Plugins/Editor/ModelingToolsEditorMode/Source/ModelingToolsEditorMode/Public/ModelingSelectionInteraction.h
new file mode 100644
--- /dev/null
+++ b/Engine/Plugins/Editor/ModelingToolsEditorMode/Source/ModelingToolsEditorMode/Public/ModelingSelectionInteraction.h
@@ -0,0 +1,142 @@
+// Copyright Epic Games, Inc. All Rights Reserved.
+
+
+#pragma once
+
+#include "BaseBehaviors/BehaviorTargetInterfaces.h"
+#include "BaseBehaviors/SingleClickBehavior.h"
+#include "InputBehaviorSet.h"
+#include "BaseGizmos/TransformGizmoUtil.h"
+#include "FrameTypes.h"
+#include "ModelingSelectionInteraction.generated.h"
+
+class UGeometrySelectionManager;
+
+/**
+ * UModelingSelectionInteraction provides element-level selection behavior (ie mesh triangles/edges/vertices)
+ * via the UGeometrySelectionManager. The Interaction doesn't actually know about the meshes or interact
+ * with them at all, it simply "operates" the SelectionManager based on user input. 
+ * 
+ * In addition to handling selection/deselection behaviors based on user input, the UModelingSelectionInteraction
+ * also creates a Gizmo for the active selection if it is transformable, and forwards gizmo transformations
+ * to the SelectionManager.
+ * 
+ * Currently some desired Selection-related interactions require a slightly convoluted path through the code,
+ * because (for example) we want the standard UE editor gizmo to have "click precedence" over our selection
+ * click behaviors. However we will get the input event first, so we need to be able to check if the gizmo is
+ * hit via a callback function provided by the EdMode owning the ToolsContext this Interaction will be registered in.
+ * Similarly we need to be able to filter out hits if a closer visible object is hit, because there is
+ * no higher-level Behavior that will do that for us. (This may be improved in future, if more of the existing
+ * Editor interaction behaviors are re-implemented as InputBehaviors)
+ * 
+ */
+UCLASS()
+class UModelingSelectionInteraction : public UObject, public IInputBehaviorSource, public IClickBehaviorTarget
+{
+	GENERATED_BODY()
+
+public:
+
+	/**
+	 * Set up the Interaction, creates and registers Behaviors/etc. 
+	 * @param CanChangeSelectionCallbackIn this function will be called to determine if the current Selection is allowed to be modified (for example, when a Tool is active, we may wish to lock selection)
+	 */
+	virtual void Initialize(
+		TObjectPtr<UGeometrySelectionManager> SelectionManager,
+		TUniqueFunction<bool()> CanChangeSelectionCallbackIn,
+		TUniqueFunction<bool(const FInputDeviceRay&)> ExternalHitCaptureCallbackIn);
+
+	virtual void Shutdown();
+
+	// this needs to be called from EdMode tick or rendering, to avoid mouse-interrupt priority issues
+	virtual void ApplyPendingTransformInteractions();
+
+	DECLARE_MULTICAST_DELEGATE(FOnBeginTransformInteraction);
+	FOnBeginTransformInteraction OnTransformBegin;
+
+	DECLARE_MULTICAST_DELEGATE(FOnEndTransformInteraction);
+	FOnEndTransformInteraction OnTransformEnd;
+
+public:
+	// click-to-select behavior
+	UPROPERTY()
+	USingleClickInputBehavior* ClickBehavior;
+
+	// set of all behaviors, will be passed up to UInputRouter
+	UPROPERTY()
+	UInputBehaviorSet* BehaviorSet;
+
+public:
+	//
+	// IInputBehaviorSource API
+	//
+	virtual const UInputBehaviorSet* GetInputBehaviors() const
+	{
+		return BehaviorSet;
+	}
+
+	//
+	// IClickBehaviorTarget implementation
+	//
+	virtual FInputRayHit IsHitByClick(const FInputDeviceRay& ClickPos) override;
+	virtual void OnClicked(const FInputDeviceRay& ClickPos) override;
+
+
+	//
+	// IModifierToggleBehaviorTarget implementation
+	//
+	virtual void OnUpdateModifierState(int ModifierID, bool bIsOn) override;
+
+
+protected:
+
+	// default change-selection callback always allows selection change
+	TUniqueFunction<bool()> CanChangeSelectionCallback = []() { return true; };
+
+	// 
+	TUniqueFunction<bool(const FInputDeviceRay&)> ExternalHitCaptureCallback = [](const FInputDeviceRay&) { return false; };
+
+	void ComputeSceneHits(const FInputDeviceRay& ClickPos,
+		bool& bHitActiveObjects, FInputRayHit& ActiveObjectHit,
+		bool& bHitInactiveObjectFirst, FInputRayHit& InactiveObjectHit);
+
+	bool bClearSelectionOnClicked = false;
+
+	// flags used to identify behavior modifier keys/buttons
+	static const int AddToSelectionModifier = 1;
+	bool bAddToSelectionEnabled = false;
+
+	static const int ToggleSelectionModifier = 2;
+	bool bToggleSelectionEnabled = false;
+
+
+
+	UPROPERTY()
+	TObjectPtr<UGeometrySelectionManager> SelectionManager;
+
+	FDelegateHandle OnSelectionModifiedDelegate;
+	void OnSelectionManager_SelectionModified();
+
+
+
+
+	UPROPERTY()
+	TObjectPtr<UTransformProxy> TransformProxy;
+
+	UPROPERTY()
+	TObjectPtr<UCombinedTransformGizmo> TransformGizmo;
+
+	UE::Geometry::FFrame3d InitialGizmoFrame;
+	FVector3d InitialGizmoScale;
+	UE::Geometry::FFrame3d LastUpdateGizmoFrame;
+	FVector3d LastUpdateGizmoScale;
+	bool bInActiveTransform = false;
+
+	void UpdateGizmoOnSelectionChange();
+	bool bGizmoUpdatePending = false;
+
+	void OnBeginGizmoTransform(UTransformProxy* Proxy);
+	void OnEndGizmoTransform(UTransformProxy* Proxy);
+	void OnGizmoTransformChanged(UTransformProxy* Proxy, FTransform Transform);
+
+};
\ No newline at end of file
diff --git a/Engine/Plugins/Editor/ModelingToolsEditorMode/Source/ModelingToolsEditorMode/Public/ModelingToolsEditorMode.h b/Engine/Plugins/Editor/ModelingToolsEditorMode/Source/ModelingToolsEditorMode/Public/ModelingToolsEditorMode.h
--- a/Engine/Plugins/Editor/ModelingToolsEditorMode/Source/ModelingToolsEditorMode/Public/ModelingToolsEditorMode.h
+++ b/Engine/Plugins/Editor/ModelingToolsEditorMode/Source/ModelingToolsEditorMode/Public/ModelingToolsEditorMode.h
@@ -17,6 +17,8 @@ class FUICommandList;
 class FStylusStateTracker;		// for stylus events
 class FLevelObjectsObserver;
 class UModelingSceneSnappingManager;
+class UModelingSelectionInteraction;
+class UGeometrySelectionManager;
 
 UCLASS(Transient)
 class UModelingToolsEditorMode : public UBaseLegacyWidgetEdMode
@@ -56,18 +58,25 @@ public:
 	// called when we "end" this editor mode (ie switch to another tab)
 	virtual void Exit() override;
 
-
 	virtual bool ShouldToolStartBeAllowed(const FString& ToolIdentifier) const override;
 
 	//////////////////
 	// End of UEdMode interface
 	//////////////////
 
+
+	virtual UGeometrySelectionManager* GetSelectionManager() const
+	{
+		return SelectionManager;
+	}
+
 protected:
 	virtual void BindCommands() override;
 	virtual void CreateToolkit() override;
 	virtual void OnToolStarted(UInteractiveToolManager* Manager, UInteractiveTool* Tool) override;
 	virtual void OnToolEnded(UInteractiveToolManager* Manager, UInteractiveTool* Tool) override;
+	
+	virtual void OnToolPostBuild(UInteractiveToolManager* InToolManager, EToolSide InSide, UInteractiveTool* InBuiltTool, UInteractiveToolBuilder* InToolBuilder, const FToolBuilderState& ToolState);
 
 	FDelegateHandle MeshCreatedEventHandle;
 	FDelegateHandle TextureCreatedEventHandle;
@@ -78,7 +87,22 @@ protected:
 	TSharedPtr<FLevelObjectsObserver> LevelObjectsObserver;
 
 	UPROPERTY()
-	UModelingSceneSnappingManager* SceneSnappingManager;
+	TObjectPtr<UModelingSceneSnappingManager> SceneSnappingManager;
+
+	UPROPERTY()
+	TObjectPtr<UGeometrySelectionManager> SelectionManager;
+
+	UPROPERTY()
+	TObjectPtr<UModelingSelectionInteraction> SelectionInteraction;
+
+	FDelegateHandle SelectionManager_SelectionModifiedHandle;
+
+	bool GetGeometrySelectionChangesAllowed() const;
+	bool TestForEditorGizmoHit(const FInputDeviceRay&) const;
+
+	void UpdateSelectionManagerOnEditorSelectionChange();
+
+	void OnToolsContextRender(IToolsContextRenderAPI* RenderAPI);
 
 	void ModelingModeShortcutRequested(EModelingModeActionCommands Command);
 	void FocusCameraAtCursorHotkey();
@@ -88,6 +112,7 @@ protected:
 
 	void ConfigureRealTimeViewportsOverride(bool bEnable);
 
+
 	// analytics tracking
 	static FDateTime LastModeStartTimestamp;
 	static FDateTime LastToolStartTimestamp;
diff --git a/Engine/Plugins/Editor/ModelingToolsEditorMode/Source/ModelingToolsEditorMode/Public/ModelingToolsEditorModeToolkit.h b/Engine/Plugins/Editor/ModelingToolsEditorMode/Source/ModelingToolsEditorMode/Public/ModelingToolsEditorModeToolkit.h
--- a/Engine/Plugins/Editor/ModelingToolsEditorMode/Source/ModelingToolsEditorMode/Public/ModelingToolsEditorModeToolkit.h
+++ b/Engine/Plugins/Editor/ModelingToolsEditorMode/Source/ModelingToolsEditorMode/Public/ModelingToolsEditorModeToolkit.h
@@ -66,9 +66,10 @@ public:
 	virtual FText GetRestrictiveModeAutoGeneratedAssetPathText() const;
 	virtual void OnRestrictiveModeAutoGeneratedAssetPathTextCommitted(const FText& InNewText, ETextCommit::Type InTextCommit) const;
 
+	virtual void ForceToolPaletteRebuild();
+
 private:
 	const static TArray<FName> PaletteNames_Standard;
-	const static TArray<FName> PaletteNames_Experimental;
 
 	FText ActiveToolName;
 	FText ActiveToolMessage;
@@ -113,5 +114,9 @@ private:
 
 	bool bFirstInitializeAfterModeSetup = true;
 
+
+	bool bShowActiveSelectionActions = false;
+
+
 	// custom accept/cancel/complete handlers
 };
diff --git a/Engine/Plugins/Editor/ModelingToolsEditorMode/Source/ModelingToolsEditorMode/Public/ModelingToolsManagerActions.h b/Engine/Plugins/Editor/ModelingToolsEditorMode/Source/ModelingToolsEditorMode/Public/ModelingToolsManagerActions.h
--- a/Engine/Plugins/Editor/ModelingToolsEditorMode/Source/ModelingToolsEditorMode/Public/ModelingToolsManagerActions.h
+++ b/Engine/Plugins/Editor/ModelingToolsEditorMode/Source/ModelingToolsEditorMode/Public/ModelingToolsManagerActions.h
@@ -141,6 +141,16 @@ public:
 	TSharedPtr<FUICommandInfo> BeginPolyModelTool_Outset;
 	TSharedPtr<FUICommandInfo> BeginPolyModelTool_CutFaces;
 
+
+	TSharedPtr<FUICommandInfo> BeginSelectionAction_ToVertexType;
+	TSharedPtr<FUICommandInfo> BeginSelectionAction_ToEdgeType;
+	TSharedPtr<FUICommandInfo> BeginSelectionAction_ToFaceType;
+
+	TSharedPtr<FUICommandInfo> BeginSelectionAction_ToObjectType;
+	TSharedPtr<FUICommandInfo> BeginSelectionAction_ToTriangleType;
+	TSharedPtr<FUICommandInfo> BeginSelectionAction_ToPolygroupType;
+
+
 	//
 	// Accept/Cancel/Complete commands are used to end the active Tool via ToolManager
 	//
@@ -157,3 +167,4 @@ public:
 	 */
 	virtual void RegisterCommands() override;
 };
+
diff --git a/Engine/Plugins/Runtime/MeshModelingToolset/Source/ModelingComponents/Public/Selection/GeometrySelectionManager.h b/Engine/Plugins/Runtime/MeshModelingToolset/Source/ModelingComponents/Public/Selection/GeometrySelectionManager.h
--- a/Engine/Plugins/Runtime/MeshModelingToolset/Source/ModelingComponents/Public/Selection/GeometrySelectionManager.h
+++ b/Engine/Plugins/Runtime/MeshModelingToolset/Source/ModelingComponents/Public/Selection/GeometrySelectionManager.h
@@ -253,7 +253,7 @@ protected:
 	EGeometryElementType SelectionElementType = EGeometryElementType::Face;
 	void SetSelectionElementTypeInternal(EGeometryElementType NewElementType);
 
-	EMeshTopologyMode MeshTopologyMode = EMeshTopologyMode::Polygroup;
+	EMeshTopologyMode MeshTopologyMode = EMeshTopologyMode::None;
 	void SetMeshTopologyModeInternal(EMeshTopologyMode NewTopologyMode);
 
 	// ITF references
