ModelingMode: fix live coding

diff --git a/Engine/Plugins/Editor/ModelingToolsEditorMode/Source/ModelingToolsEditorMode/Private/ModelingToolsEditorMode.cpp b/Engine/Plugins/Editor/ModelingToolsEditorMode/Source/ModelingToolsEditorMode/Private/ModelingToolsEditorMode.cpp
--- a/Engine/Plugins/Editor/ModelingToolsEditorMode/Source/ModelingToolsEditorMode/Private/ModelingToolsEditorMode.cpp
+++ b/Engine/Plugins/Editor/ModelingToolsEditorMode/Source/ModelingToolsEditorMode/Private/ModelingToolsEditorMode.cpp
@@ -149,6 +149,7 @@ UModelingToolsEditorMode::UModelingToolsEditorMode()
 }
 
 UModelingToolsEditorMode::UModelingToolsEditorMode(FVTableHelper& Helper)
+	: UBaseLegacyWidgetEdMode(Helper)
 {
 }
 
