nonunity fix

diff --git a/Engine/Plugins/Editor/ModelingToolsEditorMode/Source/ModelingToolsEditorMode/Private/ModelingToolsEditorMode.cpp b/Engine/Plugins/Editor/ModelingToolsEditorMode/Source/ModelingToolsEditorMode/Private/ModelingToolsEditorMode.cpp
--- a/Engine/Plugins/Editor/ModelingToolsEditorMode/Source/ModelingToolsEditorMode/Private/ModelingToolsEditorMode.cpp
+++ b/Engine/Plugins/Editor/ModelingToolsEditorMode/Source/ModelingToolsEditorMode/Private/ModelingToolsEditorMode.cpp
@@ -120,6 +120,7 @@
 #include "ModelingModeAssetUtils.h"
 #include "EditorModelingObjectsCreationAPI.h"
 
+#include "Selections/GeometrySelection.h"
 #include "Selection/GeometrySelectionManager.h"
 #include "Selection/DynamicMeshSelector.h"
 #include "ModelingSelectionInteraction.h"
@@ -721,7 +722,7 @@ void UModelingToolsEditorMode::Enter()
 		RegisterSelectionTopologyType(UGeometrySelectionManager::EMeshTopologyMode::Polygroup, ToolManagerCommands.BeginSelectionAction_ToPolygroupType);
 	}
 
-	auto RegisterSelectionElementType = [this](EGeometryElementType ElementMode, TSharedPtr<FUICommandInfo> UICommand)
+	auto RegisterSelectionElementType = [this](UE::Geometry::EGeometryElementType ElementMode, TSharedPtr<FUICommandInfo> UICommand)
 	{
 		Toolkit->GetToolkitCommands()->MapAction(UICommand,
 			FExecuteAction::CreateLambda([this, ElementMode]() { GetSelectionManager()->SetSelectionElementType(ElementMode); }),
@@ -731,9 +732,9 @@ void UModelingToolsEditorMode::Enter()
 	};
 	if (GetSelectionManager() != nullptr)
 	{
-		RegisterSelectionElementType(EGeometryElementType::Vertex, ToolManagerCommands.BeginSelectionAction_ToVertexType);
-		RegisterSelectionElementType(EGeometryElementType::Edge, ToolManagerCommands.BeginSelectionAction_ToEdgeType);
-		RegisterSelectionElementType(EGeometryElementType::Face, ToolManagerCommands.BeginSelectionAction_ToFaceType);
+		RegisterSelectionElementType(UE::Geometry::EGeometryElementType::Vertex, ToolManagerCommands.BeginSelectionAction_ToVertexType);
+		RegisterSelectionElementType(UE::Geometry::EGeometryElementType::Edge, ToolManagerCommands.BeginSelectionAction_ToEdgeType);
+		RegisterSelectionElementType(UE::Geometry::EGeometryElementType::Face, ToolManagerCommands.BeginSelectionAction_ToFaceType);
 	}
 
 
@@ -1272,7 +1273,7 @@ bool UModelingToolsEditorMode::ComputeBoundingBoxForViewportFocus(AActor* Actor,
 	// if we have an active Selection we can focus on that
 	if (GetSelectionManager() && GetSelectionManager()->HasSelection())
 	{
-		FGeometrySelectionBounds SelectionBounds;
+		UE::Geometry::FGeometrySelectionBounds SelectionBounds;
 		GetSelectionManager()->GetSelectionBounds(SelectionBounds);
 		InOutBox = (FBox)SelectionBounds.WorldBounds;
 		ProcessFocusBoxFunc(InOutBox);
