SelectionManager: modify Shutdown() to handle clearing active targets itself, as calling ClearActiveTargets() on editor shutdown can result in error if Component is already destroyed

diff --git a/Engine/Plugins/Editor/ModelingToolsEditorMode/Source/ModelingToolsEditorMode/Private/ModelingToolsEditorMode.cpp b/Engine/Plugins/Editor/ModelingToolsEditorMode/Source/ModelingToolsEditorMode/Private/ModelingToolsEditorMode.cpp
--- a/Engine/Plugins/Editor/ModelingToolsEditorMode/Source/ModelingToolsEditorMode/Private/ModelingToolsEditorMode.cpp
+++ b/Engine/Plugins/Editor/ModelingToolsEditorMode/Source/ModelingToolsEditorMode/Private/ModelingToolsEditorMode.cpp
@@ -889,8 +889,7 @@ void UModelingToolsEditorMode::Exit()
 	{
 		SelectionManager->OnSelectionModified.Remove(SelectionManager_SelectionModifiedHandle);
 		SelectionManager->ClearSelection();
-		SelectionManager->ClearActiveTargets();
-		SelectionManager->Shutdown();
+		SelectionManager->Shutdown();		// will clear active targets
 
 		// hopefully temporary...remove SelectionManager from ContextObjectStore
 		GetInteractiveToolsContext()->ContextObjectStore->RemoveContextObject(SelectionManager);
diff --git a/Engine/Plugins/Runtime/MeshModelingToolset/Source/ModelingComponents/Private/Selection/GeometrySelectionManager.cpp b/Engine/Plugins/Runtime/MeshModelingToolset/Source/ModelingComponents/Private/Selection/GeometrySelectionManager.cpp
--- a/Engine/Plugins/Runtime/MeshModelingToolset/Source/ModelingComponents/Private/Selection/GeometrySelectionManager.cpp
+++ b/Engine/Plugins/Runtime/MeshModelingToolset/Source/ModelingComponents/Private/Selection/GeometrySelectionManager.cpp
@@ -39,6 +39,9 @@ void UGeometrySelectionManager::Shutdown()
 	TransactionsAPI = nullptr;
 
 	ResetTargetCache();
+	ActiveTargetReferences.Reset();
+	ActiveTargetMap.Reset();
+	UpdateSelectionRenderCacheOnTargetChange();
 }
 
 bool UGeometrySelectionManager::HasBeenShutDown() const
