ModelingMode: Replace usage of deprecated function UObjectBaseUtility::GetOutermost() with UObjectBaseUtility::GetPackage()

diff --git a/Engine/Plugins/Editor/ModelingToolsEditorMode/Source/ModelingToolsEditorMode/Private/ModelingModeAssetUtils.cpp b/Engine/Plugins/Editor/ModelingToolsEditorMode/Source/ModelingToolsEditorMode/Private/ModelingModeAssetUtils.cpp
--- a/Engine/Plugins/Editor/ModelingToolsEditorMode/Source/ModelingToolsEditorMode/Private/ModelingModeAssetUtils.cpp
+++ b/Engine/Plugins/Editor/ModelingToolsEditorMode/Source/ModelingToolsEditorMode/Private/ModelingModeAssetUtils.cpp
@@ -73,12 +73,12 @@ FString UE::Modeling::GetWorldRelativeAssetRootPath(const UWorld* World)
 		World = GEditor->bIsSimulatingInEditor ? GEditor->GetPIEWorldContext()->World() : GEditor->GetEditorWorldContext().World();
 	}
 
-	if (ensure(World->GetOutermost() != nullptr) == false)
+	if (ensure(World->GetPackage() != nullptr) == false)
 	{
 		return TEXT("/Game/");
 	}
 
-	const FString WorldPackageName = World->GetOutermost()->GetName();
+	const FString WorldPackageName = World->GetPackage()->GetName();
 	const FString WorldPackageFolder = FPackageName::GetLongPackagePath(WorldPackageName);
 	return WorldPackageFolder;
 }
