MeshModelingToolset: Add Pattern tool icon

diff --git a/Engine/Plugins/Editor/ModelingToolsEditorMode/Source/ModelingToolsEditorMode/Private/ModelingToolsEditorModeStyle.cpp b/Engine/Plugins/Editor/ModelingToolsEditorMode/Source/ModelingToolsEditorMode/Private/ModelingToolsEditorModeStyle.cpp
--- a/Engine/Plugins/Editor/ModelingToolsEditorMode/Source/ModelingToolsEditorMode/Private/ModelingToolsEditorModeStyle.cpp
+++ b/Engine/Plugins/Editor/ModelingToolsEditorMode/Source/ModelingToolsEditorMode/Private/ModelingToolsEditorModeStyle.cpp
@@ -279,8 +279,8 @@ void FModelingToolsEditorModeStyle::Initialize()
 		StyleSet->Set("ModelingToolsManagerCommands.BeginConvertMeshesTool.Small", new IMAGE_BRUSH_SVG("Icons/Convert_20", Icon20x20));
 		StyleSet->Set("ModelingToolsManagerCommands.BeginSplitMeshesTool", new IMAGE_BRUSH_SVG("Icons/GeometrySplit", Icon20x20));
 		StyleSet->Set("ModelingToolsManagerCommands.BeginSplitMeshesTool.Small", new IMAGE_BRUSH_SVG("Icons/GeometrySplit", Icon20x20));
-
-		
+		StyleSet->Set("ModelingToolsManagerCommands.BeginPatternTool", new IMAGE_BRUSH_SVG("Icons/ModelingPattern", Icon20x20));
+		StyleSet->Set("ModelingToolsManagerCommands.BeginPatternTool.Small", new IMAGE_BRUSH_SVG("Icons/ModelingPattern", Icon20x20));
 
 		StyleSet->Set("ModelingToolsManagerCommands.BeginVolumeToMeshTool",                  new IMAGE_PLUGIN_BRUSH("Icons/ModelingVol2Mesh_x40",         Icon20x20));
 		StyleSet->Set("ModelingToolsManagerCommands.BeginVolumeToMeshTool.Small",            new IMAGE_PLUGIN_BRUSH("Icons/ModelingVol2Mesh_x40",         Icon20x20));
diff --git a/Engine/Plugins/Editor/ModelingToolsEditorMode/Source/ModelingToolsEditorMode/Private/ModelingToolsManagerActions.cpp b/Engine/Plugins/Editor/ModelingToolsEditorMode/Source/ModelingToolsEditorMode/Private/ModelingToolsManagerActions.cpp
--- a/Engine/Plugins/Editor/ModelingToolsEditorMode/Source/ModelingToolsEditorMode/Private/ModelingToolsManagerActions.cpp
+++ b/Engine/Plugins/Editor/ModelingToolsEditorMode/Source/ModelingToolsEditorMode/Private/ModelingToolsManagerActions.cpp
@@ -97,7 +97,7 @@ void FModelingToolsManagerCommands::RegisterCommands()
 	REGISTER_MODELING_TOOL_COMMAND(BeginTransferMeshTool, "Transfer", "Transfer meshes");
 	REGISTER_MODELING_TOOL_COMMAND(BeginConvertMeshesTool, "Convert", "Convert meshes");
 	REGISTER_MODELING_TOOL_COMMAND(BeginSplitMeshesTool, "Split", "Split meshes");
-	REGISTER_MODELING_TOOL_COMMAND(BeginPatternTool, "Pattern", "Create Patterns of meshes");
+	REGISTER_MODELING_TOOL_COMMAND(BeginPatternTool, "Pattern", "Create patterns of meshes");
 
 	// MeshOps
 	REGISTER_MODELING_TOOL_COMMAND(BeginSimplifyMeshTool, "Simplfy", "Simplify mesh");
