UVEditor: Add support for new launcher button to appear in the Modeling Mode UV tools category. Introduces a new GeometryProcessingInterface to connect Modeling Mode to the UV Editor plugin while allowing the UV Editor plugin to be optional still.

diff --git a/Engine/Plugins/Editor/ModelingToolsEditorMode/Source/ModelingToolsEditorMode/ModelingToolsEditorMode.Build.cs b/Engine/Plugins/Editor/ModelingToolsEditorMode/Source/ModelingToolsEditorMode/ModelingToolsEditorMode.Build.cs
--- a/Engine/Plugins/Editor/ModelingToolsEditorMode/Source/ModelingToolsEditorMode/ModelingToolsEditorMode.Build.cs
+++ b/Engine/Plugins/Editor/ModelingToolsEditorMode/Source/ModelingToolsEditorMode/ModelingToolsEditorMode.Build.cs
@@ -48,6 +48,7 @@ public class ModelingToolsEditorMode : ModuleRules
 				"EditorInteractiveToolsFramework",
 				"GeometryFramework",
 				"GeometryCore",
+				"GeometryProcessingInterfaces",
 				"DynamicMesh",
 				"ModelingComponents",
 				"ModelingComponentsEditorOnly",
diff --git a/Engine/Plugins/Editor/ModelingToolsEditorMode/Source/ModelingToolsEditorMode/Private/ModelingToolsEditorMode.cpp b/Engine/Plugins/Editor/ModelingToolsEditorMode/Source/ModelingToolsEditorMode/Private/ModelingToolsEditorMode.cpp
--- a/Engine/Plugins/Editor/ModelingToolsEditorMode/Source/ModelingToolsEditorMode/Private/ModelingToolsEditorMode.cpp
+++ b/Engine/Plugins/Editor/ModelingToolsEditorMode/Source/ModelingToolsEditorMode/Private/ModelingToolsEditorMode.cpp
@@ -1,6 +1,8 @@
 // Copyright Epic Games, Inc. All Rights Reserved.
 
 #include "ModelingToolsEditorMode.h"
+#include "IGeometryProcessingInterfacesModule.h"
+#include "GeometryProcessingInterfaces/UVEditorAssetEditor.h"
 #include "InteractiveTool.h"
 #include "ModelingToolsEditorModeToolkit.h"
 #include "ModelingToolsEditorModeSettings.h"
@@ -643,6 +645,8 @@ void UModelingToolsEditorMode::Enter()
 
 	RegisterTool(ToolManagerCommands.BeginUVSeamEditTool, TEXT("BeginUVSeamEditTool"), NewObject< USeamSculptToolBuilder>());
 
+	RegisterUVEditor();
+
 	auto MeshSelectionToolBuilder = NewObject<UMeshSelectionToolBuilder>();
 	RegisterTool(ToolManagerCommands.BeginMeshSelectionTool, TEXT("BeginMeshSelectionTool"), MeshSelectionToolBuilder);
 
@@ -855,6 +859,63 @@ void UModelingToolsEditorMode::Enter()
 
 }
 
+void UModelingToolsEditorMode::RegisterUVEditor()
+{
+	// Handle the inclusion of the optional UVEditor button if the UVEditor plugin has been found
+	IGeometryProcessingInterfacesModule& GeomProcInterfaces = FModuleManager::Get().LoadModuleChecked<IGeometryProcessingInterfacesModule>("GeometryProcessingInterfaces");
+	IGeometryProcessing_UVEditorAssetEditor* UVEditorAPI = GeomProcInterfaces.GetUVEditorAssetEditorImplementation();
+	const FModelingToolsManagerCommands& ToolManagerCommands = FModelingToolsManagerCommands::Get();
+
+	if (UVEditorAPI)
+	{
+		const TSharedRef<FUICommandList>& CommandList = Toolkit->GetToolkitCommands();
+		CommandList->MapAction(ToolManagerCommands.LaunchUVEditor,
+			FExecuteAction::CreateLambda([this, UVEditorAPI]() 
+            {
+				if (UVEditorAPI) 
+                {
+					EToolsContextScope ToolScope = GetDefaultToolScope();
+					UEditorInteractiveToolsContext* UseToolsContext = GetInteractiveToolsContext(ToolScope);
+					if (ensure(UseToolsContext != nullptr) == false)
+					{
+						return;
+					}
+
+					TArray<UObject*> SelectedActors, SelectedComponents;
+					TArray<TObjectPtr<UObject>> SelectedObjects;
+					UseToolsContext->GetParentEditorModeManager()->GetSelectedActors()->GetSelectedObjects(SelectedActors);
+					UseToolsContext->GetParentEditorModeManager()->GetSelectedComponents()->GetSelectedObjects(SelectedComponents);
+					SelectedObjects.Append(SelectedActors);
+					SelectedObjects.Append(SelectedComponents);
+					UVEditorAPI->LaunchUVEditor(SelectedObjects);
+				}
+				}),
+			FCanExecuteAction::CreateLambda([this, UVEditorAPI]() 
+            {
+				if (UVEditorAPI) 
+                {
+					EToolsContextScope ToolScope = GetDefaultToolScope();
+					UEditorInteractiveToolsContext* UseToolsContext = GetInteractiveToolsContext(ToolScope);
+					if (ensure(UseToolsContext != nullptr) == false)
+					{
+						return false;
+					}
+
+					TArray<UObject*> SelectedActors, SelectedComponents;
+					TArray<TObjectPtr<UObject>> SelectedObjects;
+					UseToolsContext->GetParentEditorModeManager()->GetSelectedActors()->GetSelectedObjects(SelectedActors);
+					UseToolsContext->GetParentEditorModeManager()->GetSelectedComponents()->GetSelectedObjects(SelectedComponents);
+					SelectedObjects.Append(SelectedActors);
+					SelectedObjects.Append(SelectedComponents);
+					return UVEditorAPI->CanLaunchUVEditor(SelectedObjects);
+				}
+				return false;
+				})
+			);
+	}
+
+}
+
 
 
 void UModelingToolsEditorMode::Exit()
diff --git a/Engine/Plugins/Editor/ModelingToolsEditorMode/Source/ModelingToolsEditorMode/Private/ModelingToolsEditorModeStyle.cpp b/Engine/Plugins/Editor/ModelingToolsEditorMode/Source/ModelingToolsEditorMode/Private/ModelingToolsEditorModeStyle.cpp
--- a/Engine/Plugins/Editor/ModelingToolsEditorMode/Source/ModelingToolsEditorMode/Private/ModelingToolsEditorModeStyle.cpp
+++ b/Engine/Plugins/Editor/ModelingToolsEditorMode/Source/ModelingToolsEditorMode/Private/ModelingToolsEditorModeStyle.cpp
@@ -271,6 +271,9 @@ void FModelingToolsEditorModeStyle::Initialize()
 		StyleSet->Set("ModelingToolsManagerCommands.BeginTransformUVIslandsTool.Small",      new IMAGE_PLUGIN_BRUSH("Icons/TransformUVs_40x",     Icon20x20));         
 		StyleSet->Set("ModelingToolsManagerCommands.BeginUVLayoutTool",                      new IMAGE_PLUGIN_BRUSH("Icons/UVLayout_40x",         Icon20x20));    
 		StyleSet->Set("ModelingToolsManagerCommands.BeginUVLayoutTool.Small",                new IMAGE_PLUGIN_BRUSH("Icons/UVLayout_40x",         Icon20x20));    
+		StyleSet->Set("ModelingToolsManagerCommands.LaunchUVEditor",                         new IMAGE_BRUSH_SVG("Icons/UVEditor", Icon20x20));
+		StyleSet->Set("ModelingToolsManagerCommands.LaunchUVEditor.Small",                   new IMAGE_BRUSH_SVG("Icons/UVEditor", Icon20x20));
+
 		StyleSet->Set("ModelingToolsManagerCommands.BeginMeshGroupPaintTool", new IMAGE_BRUSH_SVG("Icons/GroupPaint", Icon20x20));
 		StyleSet->Set("ModelingToolsManagerCommands.BeginMeshGroupPaintTool.Small", new IMAGE_BRUSH_SVG("Icons/GroupPaint", Icon20x20));
 		StyleSet->Set("ModelingToolsManagerCommands.BeginLatticeDeformerTool", new IMAGE_BRUSH_SVG("Icons/LatticeDeformation", Icon20x20));
diff --git a/Engine/Plugins/Editor/ModelingToolsEditorMode/Source/ModelingToolsEditorMode/Private/ModelingToolsEditorModeToolkit.cpp b/Engine/Plugins/Editor/ModelingToolsEditorMode/Source/ModelingToolsEditorMode/Private/ModelingToolsEditorModeToolkit.cpp
--- a/Engine/Plugins/Editor/ModelingToolsEditorMode/Source/ModelingToolsEditorMode/Private/ModelingToolsEditorModeToolkit.cpp
+++ b/Engine/Plugins/Editor/ModelingToolsEditorMode/Source/ModelingToolsEditorMode/Private/ModelingToolsEditorModeToolkit.cpp
@@ -1,6 +1,8 @@
 // Copyright Epic Games, Inc. All Rights Reserved.
 
 #include "ModelingToolsEditorModeToolkit.h"
+#include "IGeometryProcessingInterfacesModule.h"
+#include "GeometryProcessingInterfaces/UVEditorAssetEditor.h"
 #include "ModelingToolsEditorMode.h"
 #include "ModelingToolsManagerActions.h"
 #include "ModelingToolsEditorModeSettings.h"
@@ -740,6 +742,13 @@ void FModelingToolsEditorModeToolkit::BuildToolPalette(FName PaletteIndex, class
 		ToolbarBuilder.AddToolBarButton(Commands.BeginUVSeamEditTool);
 		ToolbarBuilder.AddToolBarButton(Commands.BeginTransformUVIslandsTool);
 		ToolbarBuilder.AddToolBarButton(Commands.BeginUVLayoutTool);
+		// Handle the inclusion of the optional UVEditor button if the UVEditor plugin has been found
+		IGeometryProcessingInterfacesModule& GeomProcInterfaces = FModuleManager::Get().LoadModuleChecked<IGeometryProcessingInterfacesModule>("GeometryProcessingInterfaces");
+		IGeometryProcessing_UVEditorAssetEditor* UVEditorAPI = GeomProcInterfaces.GetUVEditorAssetEditorImplementation();
+		if (UVEditorAPI)
+		{
+			ToolbarBuilder.AddToolBarButton(Commands.LaunchUVEditor);
+		}
 	}
 	else if (PaletteIndex == VolumesTabName)
 	{
diff --git a/Engine/Plugins/Editor/ModelingToolsEditorMode/Source/ModelingToolsEditorMode/Private/ModelingToolsManagerActions.cpp b/Engine/Plugins/Editor/ModelingToolsEditorMode/Source/ModelingToolsEditorMode/Private/ModelingToolsManagerActions.cpp
--- a/Engine/Plugins/Editor/ModelingToolsEditorMode/Source/ModelingToolsEditorMode/Private/ModelingToolsManagerActions.cpp
+++ b/Engine/Plugins/Editor/ModelingToolsEditorMode/Source/ModelingToolsEditorMode/Private/ModelingToolsManagerActions.cpp
@@ -134,6 +134,8 @@ void FModelingToolsManagerCommands::RegisterCommands()
 	REGISTER_MODELING_TOOL_COMMAND(BeginUVSeamEditTool, "SeamEd", "Add UV seams to mesh");
 	REGISTER_MODELING_TOOL_COMMAND(BeginTransformUVIslandsTool, "XFormUV", "Transform UV islands in UV space");
 	REGISTER_MODELING_TOOL_COMMAND(BeginUVLayoutTool, "Layout", "Transform and pack existing UVs");
+	// This is done directly, not with the REGISTER_ macro, since we don't want it added to the tool list or use a toggle button
+	UI_COMMAND(LaunchUVEditor, "UVEditor", "Launch UV asset editor", EUserInterfaceActionType::Button, FInputChord());
 
 	// Baking
 	REGISTER_MODELING_TOOL_COMMAND(BeginBakeMeshAttributeMapsTool, "BakeTx", "Bake textures for a target mesh");
diff --git a/Engine/Plugins/Editor/ModelingToolsEditorMode/Source/ModelingToolsEditorMode/Public/ModelingToolsEditorMode.h b/Engine/Plugins/Editor/ModelingToolsEditorMode/Source/ModelingToolsEditorMode/Public/ModelingToolsEditorMode.h
--- a/Engine/Plugins/Editor/ModelingToolsEditorMode/Source/ModelingToolsEditorMode/Public/ModelingToolsEditorMode.h
+++ b/Engine/Plugins/Editor/ModelingToolsEditorMode/Source/ModelingToolsEditorMode/Public/ModelingToolsEditorMode.h
@@ -79,6 +79,9 @@ protected:
 	
 	virtual void OnToolPostBuild(UInteractiveToolManager* InToolManager, EToolSide InSide, UInteractiveTool* InBuiltTool, UInteractiveToolBuilder* InToolBuilder, const FToolBuilderState& ToolState);
 
+	// Method to optionally register the UV Editor launcher in the Modeling Mode UV category if the plugin is available.
+	void RegisterUVEditor();
+
 	FDelegateHandle MeshCreatedEventHandle;
 	FDelegateHandle TextureCreatedEventHandle;
 	FDelegateHandle SelectionModifiedEventHandle;
diff --git a/Engine/Plugins/Editor/ModelingToolsEditorMode/Source/ModelingToolsEditorMode/Public/ModelingToolsManagerActions.h b/Engine/Plugins/Editor/ModelingToolsEditorMode/Source/ModelingToolsEditorMode/Public/ModelingToolsManagerActions.h
--- a/Engine/Plugins/Editor/ModelingToolsEditorMode/Source/ModelingToolsEditorMode/Public/ModelingToolsManagerActions.h
+++ b/Engine/Plugins/Editor/ModelingToolsEditorMode/Source/ModelingToolsEditorMode/Public/ModelingToolsManagerActions.h
@@ -154,6 +154,9 @@ public:
 
 	TSharedPtr<FUICommandInfo> BeginSelectionAction_Delete;
 
+	// For connecting ModelingMode with UVEditor plugin
+	TSharedPtr<FUICommandInfo> LaunchUVEditor;
+
 	//
 	// Accept/Cancel/Complete commands are used to end the active Tool via ToolManager
 	//
diff --git a/Engine/Plugins/Experimental/UVEditor/Source/UVEditor/Private/UVEditorAssetEditorImpl.cpp b/Engine/Plugins/Experimental/UVEditor/Source/UVEditor/Private/UVEditorAssetEditorImpl.cpp
new file mode 100644
--- /dev/null
+++ b/Engine/Plugins/Experimental/UVEditor/Source/UVEditor/Private/UVEditorAssetEditorImpl.cpp
@@ -0,0 +1,65 @@
+// Copyright Epic Games, Inc. All Rights Reserved.
+
+#pragma once
+
+#include "UVEditorAssetEditorImpl.h"
+
+#include "CoreMinimal.h"
+#include "UVEditorSubsystem.h"
+#include "Editor.h"
+#include "GameFramework/Actor.h"
+
+namespace UE
+{
+namespace Geometry
+{
+	void FUVEditorAssetEditorImpl::LaunchUVEditor(const TArray<TObjectPtr<UObject>>& ObjectsIn)
+	{
+		UUVEditorSubsystem* UVSubsystem = GEditor->GetEditorSubsystem<UUVEditorSubsystem>();
+		check(UVSubsystem);
+		TArray<TObjectPtr<UObject>> ProcessedObjects;
+		ConvertInputArgsToValidTargets(ObjectsIn, ProcessedObjects);
+		UVSubsystem->StartUVEditor(ProcessedObjects);
+	}
+
+	bool FUVEditorAssetEditorImpl::CanLaunchUVEditor(const TArray<TObjectPtr<UObject>>& ObjectsIn)
+	{
+		UUVEditorSubsystem* UVSubsystem = GEditor->GetEditorSubsystem<UUVEditorSubsystem>();
+		check(UVSubsystem);
+		TArray<TObjectPtr<UObject>> ProcessedObjects;
+		ConvertInputArgsToValidTargets(ObjectsIn, ProcessedObjects);
+		return UVSubsystem->AreObjectsValidTargets(ProcessedObjects);
+	}
+
+	void FUVEditorAssetEditorImpl::ConvertInputArgsToValidTargets(const TArray<TObjectPtr<UObject>>& ObjectsIn, TArray<TObjectPtr<UObject>>& ObjectsOut) const
+	{
+		for (const TObjectPtr<UObject>& Object : ObjectsIn)
+		{
+			const AActor* Actor = Cast<const AActor>(Object);
+			if (Actor)
+			{
+				TArray<UObject*> ActorAssets;
+				Actor->GetReferencedContentObjects(ActorAssets);
+
+				if (ActorAssets.Num() > 0)
+				{
+					for (UObject* Asset : ActorAssets)
+					{
+						ObjectsOut.AddUnique(Asset);
+					}
+				}
+				else {
+					// Need to transform actors to components here because that's what the UVEditor expects to have
+					TInlineComponentArray<UActorComponent*> ActorComponents;
+					Actor->GetComponents(ActorComponents);
+					ObjectsOut.Append(ActorComponents);
+				}
+			}
+			else
+			{
+				ObjectsOut.AddUnique(Object);
+			}
+		}
+	}
+}
+}
diff --git a/Engine/Plugins/Experimental/UVEditor/Source/UVEditor/Private/UVEditorModule.cpp b/Engine/Plugins/Experimental/UVEditor/Source/UVEditor/Private/UVEditorModule.cpp
--- a/Engine/Plugins/Experimental/UVEditor/Source/UVEditor/Private/UVEditorModule.cpp
+++ b/Engine/Plugins/Experimental/UVEditor/Source/UVEditor/Private/UVEditorModule.cpp
@@ -2,6 +2,7 @@
 
 #include "UVEditorModule.h"
 
+#include "Features/IModularFeatures.h"
 #include "ContentBrowserMenuContexts.h"
 #include "EditorModeRegistry.h"
 #include "LevelEditorMenuContext.h"
@@ -20,6 +21,13 @@
 
 void FUVEditorModule::StartupModule()
 {
+	UVEditorAssetEditor = MakeShared<UE::Geometry::FUVEditorAssetEditorImpl>();
+
+	if (UVEditorAssetEditor.IsValid())
+	{
+		IModularFeatures::Get().RegisterModularFeature(IGeometryProcessing_UVEditorAssetEditor::GetModularFeatureName(), UVEditorAssetEditor.Get());
+	}
+
 	FUVEditorStyle::Get(); // Causes the constructor to be called
 	FUVEditorCommands::Register();
 
@@ -38,6 +46,12 @@ void FUVEditorModule::StartupModule()
 
 void FUVEditorModule::ShutdownModule()
 {
+	if (UVEditorAssetEditor.IsValid())
+	{
+		IModularFeatures::Get().UnregisterModularFeature(IGeometryProcessing_UVEditorAssetEditor::GetModularFeatureName(), UVEditorAssetEditor.Get());
+		UVEditorAssetEditor = nullptr;
+	}
+
 	// Clean up menu things
 	UToolMenus::UnRegisterStartupCallback(this);
 	UToolMenus::UnregisterOwner(this);
diff --git a/Engine/Plugins/Experimental/UVEditor/Source/UVEditor/Public/UVEditorAssetEditorImpl.h b/Engine/Plugins/Experimental/UVEditor/Source/UVEditor/Public/UVEditorAssetEditorImpl.h
new file mode 100644
--- /dev/null
+++ b/Engine/Plugins/Experimental/UVEditor/Source/UVEditor/Public/UVEditorAssetEditorImpl.h
@@ -0,0 +1,28 @@
+// Copyright Epic Games, Inc. All Rights Reserved.
+
+#pragma once
+
+#include "CoreMinimal.h"
+#include "GeometryProcessingInterfaces/UVEditorAssetEditor.h"
+
+namespace UE
+{
+namespace Geometry
+{
+
+/** * Implementation of IGeometryProcessing_UVEditorAssetEditor
+ */
+class UVEDITOR_API FUVEditorAssetEditorImpl : public IGeometryProcessing_UVEditorAssetEditor
+{
+public:
+	virtual void LaunchUVEditor(const TArray<TObjectPtr<UObject>>& Objects) override;
+	virtual bool CanLaunchUVEditor(const TArray<TObjectPtr<UObject>>& Objects) override;
+
+protected:
+	void ConvertInputArgsToValidTargets(const TArray<TObjectPtr<UObject>>& ObjectsIn, TArray<TObjectPtr<UObject>>& ObjectsOut) const;
+
+};
+
+
+}
+}
diff --git a/Engine/Plugins/Experimental/UVEditor/Source/UVEditor/Public/UVEditorModule.h b/Engine/Plugins/Experimental/UVEditor/Source/UVEditor/Public/UVEditorModule.h
--- a/Engine/Plugins/Experimental/UVEditor/Source/UVEditor/Public/UVEditorModule.h
+++ b/Engine/Plugins/Experimental/UVEditor/Source/UVEditor/Public/UVEditorModule.h
@@ -5,6 +5,8 @@
 #include "CoreMinimal.h"
 #include "Modules/ModuleInterface.h"
 
+#include "UVEditorAssetEditorImpl.h"
+
 class FLayoutExtender;
 
 /**
@@ -30,4 +32,6 @@ private:
 
 	/** StaticClass is not safe on shutdown, so we cache the name, and use this to unregister on shut down */
 	TArray<FName> ClassesToUnregisterOnShutdown;
+
+	TSharedPtr<UE::Geometry::FUVEditorAssetEditorImpl> UVEditorAssetEditor;
 };
diff --git a/Engine/Plugins/Experimental/UVEditor/Source/UVEditor/UVEditor.Build.cs b/Engine/Plugins/Experimental/UVEditor/Source/UVEditor/UVEditor.Build.cs
--- a/Engine/Plugins/Experimental/UVEditor/Source/UVEditor/UVEditor.Build.cs
+++ b/Engine/Plugins/Experimental/UVEditor/Source/UVEditor/UVEditor.Build.cs
@@ -31,6 +31,7 @@ public class UVEditor : ModuleRules
 				"InputCore",
 				"InteractiveToolsFramework",
 				"MeshModelingToolsExp",
+				"GeometryProcessingInterfaces" // For supporting launching the UVEditor directly from Modeling Tools or elsewhere
 				
 				// ... add other public dependencies that you statically link with here ...
 			}
diff --git a/Engine/Source/Developer/GeometryProcessingInterfaces/Private/GeometryProcessingInterfacesModule.cpp b/Engine/Source/Developer/GeometryProcessingInterfaces/Private/GeometryProcessingInterfacesModule.cpp
--- a/Engine/Source/Developer/GeometryProcessingInterfaces/Private/GeometryProcessingInterfacesModule.cpp
+++ b/Engine/Source/Developer/GeometryProcessingInterfaces/Private/GeometryProcessingInterfacesModule.cpp
@@ -6,6 +6,7 @@
 #include "Modules/ModuleManager.h"
 
 #include "GeometryProcessingInterfaces/ApproximateActors.h"
+#include "GeometryProcessingInterfaces/UVEditorAssetEditor.h"
 
 
 //DEFINE_LOG_CATEGORY_STATIC(LogMeshReduction, Verbose, All);
@@ -38,3 +39,16 @@ IGeometryProcessing_ApproximateActors* FGeometryProcessingInterfacesModule::GetA
 	return ApproximateActors;
 }
 
+IGeometryProcessing_UVEditorAssetEditor* FGeometryProcessingInterfacesModule::GetUVEditorAssetEditorImplementation()
+{
+	if (UVEditorAssetEditor == nullptr)
+	{
+		TArray<IGeometryProcessing_UVEditorAssetEditor*> UVEditorAssetEditorOptions =
+			IModularFeatures::Get().GetModularFeatureImplementations<IGeometryProcessing_UVEditorAssetEditor>(IGeometryProcessing_UVEditorAssetEditor::GetModularFeatureName());
+
+		UVEditorAssetEditor = (UVEditorAssetEditorOptions.Num() > 0) ? UVEditorAssetEditorOptions[0] : nullptr;
+	}
+
+	return UVEditorAssetEditor;
+}
+
diff --git a/Engine/Source/Developer/GeometryProcessingInterfaces/Private/GeometryProcessingInterfacesModule.h b/Engine/Source/Developer/GeometryProcessingInterfaces/Private/GeometryProcessingInterfacesModule.h
--- a/Engine/Source/Developer/GeometryProcessingInterfaces/Private/GeometryProcessingInterfacesModule.h
+++ b/Engine/Source/Developer/GeometryProcessingInterfaces/Private/GeometryProcessingInterfacesModule.h
@@ -6,6 +6,7 @@
 
 
 class IGeometryProcessing_ApproximateActors;
+class IGeometryProcessing_UVEditorAssetEditor;
 
 
 /**
@@ -26,6 +27,12 @@ public:
 	virtual IGeometryProcessing_ApproximateActors* GetApproximateActorsImplementation() override;
 
 
+	/**
+	* @return implementation of IGeometryProcessing_UVEditor, if available, or nullptr (result is cached internally)
+	*/
+	virtual IGeometryProcessing_UVEditorAssetEditor* GetUVEditorAssetEditorImplementation() override;
+
 private:
 	IGeometryProcessing_ApproximateActors* ApproximateActors = nullptr;
+	IGeometryProcessing_UVEditorAssetEditor* UVEditorAssetEditor = nullptr;
 };
diff --git a/Engine/Source/Developer/GeometryProcessingInterfaces/Public/GeometryProcessingInterfaces/UVEditorAssetEditor.h b/Engine/Source/Developer/GeometryProcessingInterfaces/Public/GeometryProcessingInterfaces/UVEditorAssetEditor.h
new file mode 100644
--- /dev/null
+++ b/Engine/Source/Developer/GeometryProcessingInterfaces/Public/GeometryProcessingInterfaces/UVEditorAssetEditor.h
@@ -0,0 +1,49 @@
+// Copyright Epic Games, Inc. All Rights Reserved.
+
+#pragma once
+
+#include "CoreMinimal.h"
+#include "Features/IModularFeature.h"
+#include "UObject/ObjectPtr.h"
+
+/**
+ * IGeometryProcessing_UVEditorAssetEditor is a generic high-level interface the UVEditor Asset Editor plugin.
+ * 
+ * This is an IModularFeature, and so clients can query/enumerate the available UVEditor implementations 
+ * based on the ::GetModularFeatureName(). However, the preferred way is to use code like the following:
+ * 
+ *     IGeometryProcessingInterfacesModule& GeomProcInterfaces = FModuleManager::Get().LoadModuleChecked<IGeometryProcessingInterfacesModule>("GeometryProcessingInterfaces");
+ *     IGeometryProcessing_UVEditorAssetEditor* UVEditorAPI = GeomProcInterfaces.GetUVEditorAssetEditorImplementation();
+ * 
+ * This will automatically determine which available implementation to use (based on any applicable config settings/etc),
+ * and cache the result of that decision.
+ */
+class IGeometryProcessing_UVEditorAssetEditor : public IModularFeature
+{
+public:
+	virtual ~IGeometryProcessing_UVEditorAssetEditor() {}
+
+	/**
+	 * Top-level driver function that clients call to launch the UVEditor Asset Editor
+	 */
+	virtual void LaunchUVEditor(const TArray<TObjectPtr<UObject>>& Objects)
+	{
+		check(false);		// not implemented in base class
+	}
+
+	/**
+	* Top-level function to check if the UVEditor can be launched at this point
+	*/
+	virtual bool CanLaunchUVEditor(const TArray<TObjectPtr<UObject>>& Objects)
+	{
+		check(false);		// not implemented in base class
+		return false;
+	}
+
+
+	// Modular feature name to register for retrieval during runtime
+	static const FName GetModularFeatureName()
+	{
+		return TEXT("GeometryProcessing_UVEditorAssetEditor");
+	}
+};
\ No newline at end of file
diff --git a/Engine/Source/Developer/GeometryProcessingInterfaces/Public/IGeometryProcessingInterfacesModule.h b/Engine/Source/Developer/GeometryProcessingInterfaces/Public/IGeometryProcessingInterfacesModule.h
--- a/Engine/Source/Developer/GeometryProcessingInterfaces/Public/IGeometryProcessingInterfacesModule.h
+++ b/Engine/Source/Developer/GeometryProcessingInterfaces/Public/IGeometryProcessingInterfacesModule.h
@@ -5,6 +5,7 @@
 #include "Modules/ModuleInterface.h"
 
 class IGeometryProcessing_ApproximateActors;
+class IGeometryProcessing_UVEditorAssetEditor;
 
 /**
  * Abstract interface to a Module that provides functions to access
@@ -27,4 +28,10 @@ public:
 	 * @return implementation of IGeometryProcessing_ApproximateActors, if available
 	 */
 	virtual IGeometryProcessing_ApproximateActors* GetApproximateActorsImplementation() = 0;
+
+	/**
+	 * @return implementation of IGeometryProcessing_UVEditorAssetEditor, if available
+	*/
+	virtual IGeometryProcessing_UVEditorAssetEditor* GetUVEditorAssetEditorImplementation() = 0;
+
 };
