RenderCaptureBaking: Add new icon for the BakeRC tool

#rnx

diff --git a/Engine/Plugins/Editor/ModelingToolsEditorMode/Source/ModelingToolsEditorMode/Private/ModelingToolsEditorModeStyle.cpp b/Engine/Plugins/Editor/ModelingToolsEditorMode/Source/ModelingToolsEditorMode/Private/ModelingToolsEditorModeStyle.cpp
--- a/Engine/Plugins/Editor/ModelingToolsEditorMode/Source/ModelingToolsEditorMode/Private/ModelingToolsEditorModeStyle.cpp
+++ b/Engine/Plugins/Editor/ModelingToolsEditorMode/Source/ModelingToolsEditorMode/Private/ModelingToolsEditorModeStyle.cpp
@@ -175,8 +175,8 @@ void FModelingToolsEditorModeStyle::Initialize()
 		StyleSet->Set("ModelingToolsManagerCommands.BeginBakeMeshAttributeMapsTool.Small", 		new IMAGE_BRUSH_SVG("Icons/BakeTexture",			Icon20x20));
 		StyleSet->Set("ModelingToolsManagerCommands.BeginBakeMultiMeshAttributeMapsTool", 			new IMAGE_BRUSH_SVG("Icons/BakeAll",			Icon20x20));
 		StyleSet->Set("ModelingToolsManagerCommands.BeginBakeMultiMeshAttributeMapsTool.Small", 		new IMAGE_BRUSH_SVG("Icons/BakeAll",			Icon20x20));
-		StyleSet->Set("ModelingToolsManagerCommands.BeginBakeRenderCaptureTool", 			new IMAGE_BRUSH_SVG("Icons/BakeAll",			Icon20x20));
-		StyleSet->Set("ModelingToolsManagerCommands.BeginBakeRenderCaptureTool.Small", 		new IMAGE_BRUSH_SVG("Icons/BakeAll",			Icon20x20));
+		StyleSet->Set("ModelingToolsManagerCommands.BeginBakeRenderCaptureTool", 			new IMAGE_BRUSH_SVG("Icons/BakeRenderCapture",			Icon20x20));
+		StyleSet->Set("ModelingToolsManagerCommands.BeginBakeRenderCaptureTool.Small", 		new IMAGE_BRUSH_SVG("Icons/BakeRenderCapture",			Icon20x20));
 		StyleSet->Set("ModelingToolsManagerCommands.BeginBakeMeshAttributeVertexTool", new IMAGE_BRUSH_SVG("Icons/BakeVertex", Icon20x20));
 		StyleSet->Set("ModelingToolsManagerCommands.BeginBakeMeshAttributeVertexTool.Small", new IMAGE_BRUSH_SVG("Icons/BakeVertex", Icon20x20));
 		//StyleSet->Set("ModelingToolsManagerCommands.BeginRemoveOccludedTrianglesTool", 				new IMAGE_PLUGIN_BRUSH("Icons/Jacket_40x",			Icon20x20));
