Show warning text when modeling tools are disabled for simulation or PIE

diff --git a/Engine/Plugins/Editor/ModelingToolsEditorMode/Source/ModelingToolsEditorMode/Private/ModelingToolsEditorMode.cpp b/Engine/Plugins/Editor/ModelingToolsEditorMode/Source/ModelingToolsEditorMode/Private/ModelingToolsEditorMode.cpp
--- a/Engine/Plugins/Editor/ModelingToolsEditorMode/Source/ModelingToolsEditorMode/Private/ModelingToolsEditorMode.cpp
+++ b/Engine/Plugins/Editor/ModelingToolsEditorMode/Source/ModelingToolsEditorMode/Private/ModelingToolsEditorMode.cpp
@@ -241,7 +241,7 @@ void UModelingToolsEditorMode::Tick(FEditorViewportClient* ViewportClient, float
 	if (Toolkit.IsValid())
 	{
 		FModelingToolsEditorModeToolkit* ModelingToolkit = (FModelingToolsEditorModeToolkit*)Toolkit.Get();
-		ModelingToolkit->EnableShowRealtimeWarning(ViewportClient->IsRealtime() == false);
+		ModelingToolkit->ShowRealtimeAndModeWarnings(ViewportClient->IsRealtime() == false);
 	}
 }
 
diff --git a/Engine/Plugins/Editor/ModelingToolsEditorMode/Source/ModelingToolsEditorMode/Private/ModelingToolsEditorModeToolkit.cpp b/Engine/Plugins/Editor/ModelingToolsEditorMode/Source/ModelingToolsEditorMode/Private/ModelingToolsEditorModeToolkit.cpp
--- a/Engine/Plugins/Editor/ModelingToolsEditorMode/Source/ModelingToolsEditorMode/Private/ModelingToolsEditorModeToolkit.cpp
+++ b/Engine/Plugins/Editor/ModelingToolsEditorMode/Source/ModelingToolsEditorMode/Private/ModelingToolsEditorModeToolkit.cpp
@@ -992,12 +992,26 @@ void FModelingToolsEditorModeToolkit::OnToolPaletteChanged(FName PaletteName)
 
 
 
-void FModelingToolsEditorModeToolkit::EnableShowRealtimeWarning(bool bEnable)
+void FModelingToolsEditorModeToolkit::ShowRealtimeAndModeWarnings(bool bShowRealtimeWarning)
 {
-	if (bShowRealtimeWarning != bEnable)
+	FText WarningText{};
+	if (GEditor->bIsSimulatingInEditor)
 	{
-		bShowRealtimeWarning = bEnable;
-		UpdateShowWarnings();
+		WarningText = LOCTEXT("ModelingModeToolkitSimulatingWarning", "Cannot use Modeling Tools while simulating.");
+	}
+	else if (GEditor->PlayWorld != NULL)
+	{
+		WarningText = LOCTEXT("ModelingModeToolkitSimulatingWarning", "Cannot use Modeling Tools in PIE.");
+	}
+	else if (bShowRealtimeWarning)
+	{
+		WarningText = LOCTEXT("ModelingModeToolkitRealtimeWarning", "Realtime Mode is required for Modeling Tools to work correctly. Please enable Realtime Mode in the Viewport Options or with the Ctrl+r hotkey.");
+	}
+	if (!WarningText.IdenticalTo(ActiveWarning))
+	{
+		ActiveWarning = WarningText;
+		ModeWarningArea->SetVisibility(ActiveWarning.IsEmpty() ? EVisibility::Collapsed : EVisibility::Visible);
+		ModeWarningArea->SetText(ActiveWarning);
 	}
 }
 
@@ -1067,23 +1081,7 @@ void FModelingToolsEditorModeToolkit::OnActiveViewportChanged(TSharedPtr<IAssetV
 	}
 }
 
-void FModelingToolsEditorModeToolkit::UpdateShowWarnings()
-{
-	if (bShowRealtimeWarning )
-	{
-		if (ModeWarningArea->GetVisibility() == EVisibility::Collapsed)
-		{
-			ModeWarningArea->SetText(LOCTEXT("ModelingModeToolkitRealtimeWarning", "Realtime Mode is required for Modeling Tools to work correctly. Please enable Realtime Mode in the Viewport Options or with the Ctrl+r hotkey."));
-			ModeWarningArea->SetVisibility(EVisibility::Visible);
-		}
-	}
-	else
-	{
-		ModeWarningArea->SetText(FText());
-		ModeWarningArea->SetVisibility(EVisibility::Collapsed);
-	}
 
-}
 
 
 void FModelingToolsEditorModeToolkit::UpdateAssetLocationMode(TSharedPtr<FString> NewString)
diff --git a/Engine/Plugins/Editor/ModelingToolsEditorMode/Source/ModelingToolsEditorMode/Public/ModelingToolsEditorModeToolkit.h b/Engine/Plugins/Editor/ModelingToolsEditorMode/Source/ModelingToolsEditorMode/Public/ModelingToolsEditorModeToolkit.h
--- a/Engine/Plugins/Editor/ModelingToolsEditorMode/Source/ModelingToolsEditorMode/Public/ModelingToolsEditorModeToolkit.h
+++ b/Engine/Plugins/Editor/ModelingToolsEditorMode/Source/ModelingToolsEditorMode/Public/ModelingToolsEditorModeToolkit.h
@@ -51,7 +51,7 @@ public:
 	virtual FText GetActiveToolDisplayName() const override { return ActiveToolName; }
 	virtual FText GetActiveToolMessage() const override { return ActiveToolMessage; }
 
-	virtual void EnableShowRealtimeWarning(bool bEnable);
+	virtual void ShowRealtimeAndModeWarnings(bool bShowRealtimeWarning);
 
 	virtual void OnToolStarted(UInteractiveToolManager* Manager, UInteractiveTool* Tool) override;
 	virtual void OnToolEnded(UInteractiveToolManager* Manager, UInteractiveTool* Tool) override;
@@ -91,8 +91,7 @@ private:
 
 	TSharedPtr<SWidget> MakeAssetConfigPanel();
 
-	bool bShowRealtimeWarning = false;
-	void UpdateShowWarnings();
+	FText ActiveWarning{};
 
 	TArray<TSharedPtr<FString>> AssetLocationModes;
 	TArray<TSharedPtr<FString>> AssetSaveModes;
