fix duplicate localization tag for modeling toolkit PIE warning

diff --git a/Engine/Plugins/Editor/ModelingToolsEditorMode/Source/ModelingToolsEditorMode/Private/ModelingToolsEditorModeToolkit.cpp b/Engine/Plugins/Editor/ModelingToolsEditorMode/Source/ModelingToolsEditorMode/Private/ModelingToolsEditorModeToolkit.cpp
--- a/Engine/Plugins/Editor/ModelingToolsEditorMode/Source/ModelingToolsEditorMode/Private/ModelingToolsEditorModeToolkit.cpp
+++ b/Engine/Plugins/Editor/ModelingToolsEditorMode/Source/ModelingToolsEditorMode/Private/ModelingToolsEditorModeToolkit.cpp
@@ -1001,7 +1001,7 @@ void FModelingToolsEditorModeToolkit::ShowRealtimeAndModeWarnings(bool bShowReal
 	}
 	else if (GEditor->PlayWorld != NULL)
 	{
-		WarningText = LOCTEXT("ModelingModeToolkitSimulatingWarning", "Cannot use Modeling Tools in PIE.");
+		WarningText = LOCTEXT("ModelingModeToolkitPIEWarning", "Cannot use Modeling Tools in PIE.");
 	}
 	else if (bShowRealtimeWarning)
 	{
