ModelingMode: fix mode panel layout issue

diff --git a/Engine/Plugins/Editor/ModelingToolsEditorMode/Source/ModelingToolsEditorMode/Private/ModelingToolsEditorModeToolkit.cpp b/Engine/Plugins/Editor/ModelingToolsEditorMode/Source/ModelingToolsEditorMode/Private/ModelingToolsEditorModeToolkit.cpp
--- a/Engine/Plugins/Editor/ModelingToolsEditorMode/Source/ModelingToolsEditorMode/Private/ModelingToolsEditorModeToolkit.cpp
+++ b/Engine/Plugins/Editor/ModelingToolsEditorMode/Source/ModelingToolsEditorMode/Private/ModelingToolsEditorModeToolkit.cpp
@@ -78,6 +78,19 @@ void FModelingToolsEditorModeToolkit::CustomizeModeDetailsViewArgs(FDetailsViewA
 	//ArgsInOut.ColumnWidth = 0.3f;
 }
 
+
+TSharedPtr<SWidget> FModelingToolsEditorModeToolkit::GetInlineContent() const
+{
+	return SNew(SVerticalBox)
+		+ SVerticalBox::Slot()
+		.FillHeight(1.0f)
+		.VAlign(VAlign_Fill)
+		[
+			ToolkitWidget.ToSharedRef()
+		];
+}
+
+
 void FModelingToolsEditorModeToolkit::Init(const TSharedPtr<IToolkitHost>& InitToolkitHost, TWeakObjectPtr<UEdMode> InOwningMode)
 {
 	// Have to create the ToolkitWidget here because FModeToolkit::Init() is going to ask for it and add
@@ -133,7 +146,7 @@ void FModelingToolsEditorModeToolkit::Init(const TSharedPtr<IToolkitHost>& InitT
 		[
 			ModeDetailsView->AsShared()
 		];
-	ToolkitWidgetVBox->AddSlot().AutoHeight().HAlign(HAlign_Fill).Padding(5)
+	ToolkitWidgetVBox->AddSlot().AutoHeight().HAlign(HAlign_Fill).VAlign(VAlign_Bottom).Padding(5)
 		[
 			MakeAssetConfigPanel()->AsShared()
 		];
@@ -813,6 +826,12 @@ void FModelingToolsEditorModeToolkit::InvokeUI()
 {
 	FModeToolkit::InvokeUI();
 
+	// FModeToolkit::UpdatePrimaryModePanel() wrapped our GetInlineContent() output in a SScrollBar widget,
+	// however this doesn't make sense as we want to dock panels to the "top" and "bottom" of our mode panel area,
+	// and the details panel in the middle has it's own scrollbar already. The SScrollBar is hardcoded as the content
+	// of FModeToolkit::InlineContentHolder so we can just replace it here
+	InlineContentHolder->SetContent(GetInlineContent().ToSharedRef());
+
 	//
 	// Apply custom section header colors.
 	// See comments below, this is done via directly manipulating Slate widgets generated deep inside BaseToolkit.cpp,
diff --git a/Engine/Plugins/Editor/ModelingToolsEditorMode/Source/ModelingToolsEditorMode/Public/ModelingToolsEditorModeToolkit.h b/Engine/Plugins/Editor/ModelingToolsEditorMode/Source/ModelingToolsEditorMode/Public/ModelingToolsEditorModeToolkit.h
--- a/Engine/Plugins/Editor/ModelingToolsEditorMode/Source/ModelingToolsEditorMode/Public/ModelingToolsEditorModeToolkit.h
+++ b/Engine/Plugins/Editor/ModelingToolsEditorMode/Source/ModelingToolsEditorMode/Public/ModelingToolsEditorModeToolkit.h
@@ -27,7 +27,7 @@ public:
 	/** IToolkit interface */
 	virtual FName GetToolkitFName() const override;
 	virtual FText GetBaseToolkitName() const override;
-	virtual TSharedPtr<class SWidget> GetInlineContent() const override { return ToolkitWidget; }
+	virtual TSharedPtr<class SWidget> GetInlineContent() const override;
 
 	// initialize toolkit widgets that need to wait until mode is initialized/entered
 	virtual void InitializeAfterModeSetup();
