MeshModelingTools: Make tool detail panel not lose focus when its contents are changed.

diff --git a/Engine/Plugins/Editor/ModelingToolsEditorMode/Source/ModelingToolsEditorMode/Private/ModelingToolsEditorModeToolkit.cpp b/Engine/Plugins/Editor/ModelingToolsEditorMode/Source/ModelingToolsEditorMode/Private/ModelingToolsEditorModeToolkit.cpp
--- a/Engine/Plugins/Editor/ModelingToolsEditorMode/Source/ModelingToolsEditorMode/Private/ModelingToolsEditorModeToolkit.cpp
+++ b/Engine/Plugins/Editor/ModelingToolsEditorMode/Source/ModelingToolsEditorMode/Private/ModelingToolsEditorModeToolkit.cpp
@@ -10,6 +10,7 @@
 #include "Selection/GeometrySelectionManager.h"
 #include "Engine/Selection.h"
 
+#include "Framework/Application/SlateApplication.h"
 #include "Framework/MultiBox/MultiBoxBuilder.h"
 #include "Modules/ModuleManager.h"
 #include "IDetailsView.h"
@@ -425,10 +426,34 @@ void FModelingToolsEditorModeToolkit::InitializeAfterModeSetup()
 void FModelingToolsEditorModeToolkit::UpdateActiveToolProperties()
 {
 	UInteractiveTool* CurTool = GetScriptableEditorMode()->GetToolManager(EToolsContextScope::EdMode)->GetActiveTool(EToolSide::Left);
-	if (CurTool != nullptr)
+	if (CurTool == nullptr)
 	{
-		ModeDetailsView->SetObjects(CurTool->GetToolProperties(true));
+		return;
 	}
+
+	// Before actually changing the detail panel, we need to see where the current keyboard focus is, because
+	// if it's inside the detail panel, we'll need to reset it to the detail panel as a whole, else we might
+	// lose it entirely when that detail panel element gets destroyed (which would make us unable to receive any
+	// hotkey presses until the user clicks somewhere).
+	TSharedPtr<SWidget> FocusedWidget = FSlateApplication::Get().GetKeyboardFocusedWidget();
+	if (FocusedWidget != ModeDetailsView) 
+	{
+		// Search upward from the currently focused widget
+		TSharedPtr<SWidget> CurrentWidget = FocusedWidget;
+		while (CurrentWidget.IsValid())
+		{
+			if (CurrentWidget == ModeDetailsView)
+			{
+				// Reset focus to the detail panel as a whole to avoid losing it when the inner elements change.
+				FSlateApplication::Get().SetKeyboardFocus(ModeDetailsView);
+				break;
+			}
+
+			CurrentWidget = CurrentWidget->GetParentWidget();
+		}
+	}
+		
+	ModeDetailsView->SetObjects(CurTool->GetToolProperties(true));
 }
 
 void FModelingToolsEditorModeToolkit::InvalidateCachedDetailPanelState(UObject* ChangedObject)
