Fix level editor gizmo lurking invisibly, still clickable, in the hitproxy buffer after tool start

diff --git a/Engine/Plugins/Editor/ModelingToolsEditorMode/Source/ModelingToolsEditorMode/Private/ModelingToolsEditorModeToolkit.cpp b/Engine/Plugins/Editor/ModelingToolsEditorMode/Source/ModelingToolsEditorMode/Private/ModelingToolsEditorModeToolkit.cpp
--- a/Engine/Plugins/Editor/ModelingToolsEditorMode/Source/ModelingToolsEditorMode/Private/ModelingToolsEditorModeToolkit.cpp
+++ b/Engine/Plugins/Editor/ModelingToolsEditorMode/Source/ModelingToolsEditorMode/Private/ModelingToolsEditorModeToolkit.cpp
@@ -44,6 +44,8 @@
 #include "Dialogs/DlgPickPath.h"
 #include "ModelingModeAssetUtils.h"
 
+#include "LevelEditorViewport.h"
+
 
 #define LOCTEXT_NAMESPACE "FModelingToolsEditorModeToolkit"
 
@@ -1082,6 +1084,16 @@ void FModelingToolsEditorModeToolkit::OnToolStarted(UInteractiveToolManager* Man
 	// disable LOD level picker once Tool is active
 	AssetLODMode->SetEnabled(false);
 	AssetLODModeLabel->SetEnabled(false);
+
+	// Invalidate all the level viewports so that e.g. hitproxy buffers are cleared
+	// (fixes the editor gizmo still being clickable despite not being visible)
+	if (GIsEditor)
+	{
+		for (FLevelEditorViewportClient* Viewport : GEditor->GetLevelViewportClients())
+		{
+			Viewport->Invalidate();
+		}
+	}
 }
 
 void FModelingToolsEditorModeToolkit::OnToolEnded(UInteractiveToolManager* Manager, UInteractiveTool* Tool)
