MeshModelingTools: Add detail customization for selection filter in PolyEd and TriEd, make loop/ring selection imply edge selection, per discussion with ux team.

diff --git a/Engine/Plugins/Editor/ModelingToolsEditorMode/Source/ModelingToolsEditorMode/Private/DetailsCustomizations/MeshTopologySelectionMechanicCustomization.cpp b/Engine/Plugins/Editor/ModelingToolsEditorMode/Source/ModelingToolsEditorMode/Private/DetailsCustomizations/MeshTopologySelectionMechanicCustomization.cpp
new file mode 100644
--- /dev/null
+++ b/Engine/Plugins/Editor/ModelingToolsEditorMode/Source/ModelingToolsEditorMode/Private/DetailsCustomizations/MeshTopologySelectionMechanicCustomization.cpp
@@ -0,0 +1,131 @@
+// Copyright Epic Games, Inc. All Rights Reserved.
+
+#include "DetailsCustomizations/MeshTopologySelectionMechanicCustomization.h"
+
+#include "DetailCategoryBuilder.h"
+#include "DetailLayoutBuilder.h"
+#include "DetailWidgetRow.h"
+#include "EditorViewportCommands.h"
+#include "Framework/MultiBox/MultiBoxBuilder.h" // FSlimHorizontalToolBarBuilder
+#include "ModelingToolsEditorModeStyle.h"
+#include "ModelingWidgets/ModelingCustomizationUtil.h"
+#include "Selection/MeshTopologySelectionMechanic.h"
+
+#define LOCTEXT_NAMESPACE "MeshTopologySelectionMechanicCustomizations"
+
+using namespace UE::ModelingUI;
+
+TSharedRef<IDetailCustomization> FMeshTopologySelectionMechanicPropertiesDetails::MakeInstance()
+{
+	return MakeShareable(new FMeshTopologySelectionMechanicPropertiesDetails);
+}
+
+void FMeshTopologySelectionMechanicPropertiesDetails::CustomizeDetails(IDetailLayoutBuilder& DetailBuilder)
+{
+	// Get the property object that we're customizing
+	TArray<TWeakObjectPtr<UObject>> ObjectsBeingCustomized;
+	DetailBuilder.GetObjectsBeingCustomized(ObjectsBeingCustomized);
+	if (!ensure(ObjectsBeingCustomized.Num() > 0))
+	{
+		return;
+	}
+	UMeshTopologySelectionMechanicProperties* Properties = CastChecked<UMeshTopologySelectionMechanicProperties>(ObjectsBeingCustomized[0]);
+	
+	// Create a toolbar for the selection filter
+	TSharedPtr<FUICommandList> CommandList = MakeShared<FUICommandList>();
+	FSlimHorizontalToolBarBuilder ToolbarBuilder(CommandList, FMultiBoxCustomization::None);
+
+	FName ToolBarStyle = "PolyEd.SelectionToolbar";
+	ToolbarBuilder.SetStyle(FModelingToolsEditorModeStyle::Get().Get(), ToolBarStyle);
+	ToolbarBuilder.SetLabelVisibility(EVisibility::Collapsed);
+
+	// Given bool property and the icon name, adds the button to the builder
+	auto AddToggleButtonForBool = [&ToolbarBuilder, Properties](TSharedPtr<IPropertyHandle> BoolPropertyHandle, FName IconName)
+	{
+		ToolbarBuilder.AddToolBarButton(FUIAction(
+			FExecuteAction::CreateWeakLambda(Properties, [BoolPropertyHandle]()
+			{
+				bool bCurrentValue;
+				BoolPropertyHandle->GetValue(bCurrentValue);
+				BoolPropertyHandle->SetValue(!bCurrentValue);
+			}),
+			FCanExecuteAction::CreateWeakLambda(Properties, [BoolPropertyHandle]() { 
+				return BoolPropertyHandle->IsEditable();
+			}),
+			FIsActionChecked::CreateWeakLambda(Properties, [BoolPropertyHandle]()
+			{
+				bool bValue;
+				BoolPropertyHandle->GetValue(bValue);
+				return bValue;
+			})),
+			NAME_None, // Extension hook
+			TAttribute<FText>(BoolPropertyHandle->GetPropertyDisplayName()), // Label
+			TAttribute<FText>(BoolPropertyHandle->GetToolTipText()), // Tooltip
+			TAttribute<FSlateIcon>(FSlateIcon(FModelingToolsEditorModeStyle::Get()->GetStyleSetName(), IconName)), 
+			EUserInterfaceActionType::ToggleButton);
+	};
+
+	ToolbarBuilder.BeginSection("SelectionFilter");
+	ToolbarBuilder.BeginBlockGroup();
+
+	TSharedPtr<IPropertyHandle> SelectVerticesHandle = DetailBuilder.GetProperty(GET_MEMBER_NAME_CHECKED(UMeshTopologySelectionMechanicProperties, bSelectVertices),
+		UMeshTopologySelectionMechanicProperties::StaticClass());
+	// This first bool we're not going to hide. Instead we're going to end up replacing it with the selection filter
+	//SelectVerticesHandle->MarkHiddenByCustomization();
+	AddToggleButtonForBool(SelectVerticesHandle, "PolyEd.SelectCorners");
+
+	TSharedPtr<IPropertyHandle> SelectEdgesHandle = DetailBuilder.GetProperty(GET_MEMBER_NAME_CHECKED(UMeshTopologySelectionMechanicProperties, bSelectEdges),
+		UMeshTopologySelectionMechanicProperties::StaticClass());
+	SelectEdgesHandle->MarkHiddenByCustomization();
+	AddToggleButtonForBool(SelectEdgesHandle, "PolyEd.SelectEdges");
+
+	TSharedPtr<IPropertyHandle> SelectFacesHandle = DetailBuilder.GetProperty(GET_MEMBER_NAME_CHECKED(UMeshTopologySelectionMechanicProperties, bSelectFaces),
+		UMeshTopologySelectionMechanicProperties::StaticClass());
+	SelectFacesHandle->MarkHiddenByCustomization();
+	AddToggleButtonForBool(SelectFacesHandle, "PolyEd.SelectFaces");
+
+	ToolbarBuilder.EndBlockGroup();
+	ToolbarBuilder.EndSection();
+
+	TSharedPtr<IPropertyHandle> SelectEdgeLoopsHandle = DetailBuilder.GetProperty(GET_MEMBER_NAME_CHECKED(UMeshTopologySelectionMechanicProperties, bSelectEdgeLoops),
+		UMeshTopologySelectionMechanicProperties::StaticClass());
+	SelectEdgeLoopsHandle->MarkHiddenByCustomization();
+	
+	TSharedPtr<IPropertyHandle> SelectEdgeRingsHandle = DetailBuilder.GetProperty(GET_MEMBER_NAME_CHECKED(UMeshTopologySelectionMechanicProperties, bSelectEdgeRings),
+		UMeshTopologySelectionMechanicProperties::StaticClass());
+	SelectEdgeRingsHandle->MarkHiddenByCustomization();
+
+	if (Properties->bDisplayPolygroupReliantControls)
+	{
+		ToolbarBuilder.BeginSection("ExtraSelectionModes");
+		ToolbarBuilder.BeginBlockGroup();
+
+		AddToggleButtonForBool(SelectEdgeLoopsHandle, "PolyEd.SelectEdgeLoops");
+		AddToggleButtonForBool(SelectEdgeRingsHandle, "PolyEd.SelectEdgeRings");
+
+		ToolbarBuilder.EndBlockGroup();
+		ToolbarBuilder.EndSection();
+	}
+
+	// Now we're going to take one of the old checkboxes and replace it with our seleciton filter. Doing this instead
+	// of creating a custom row allows us to use the original splitter, which will then align with the other splitters.
+	IDetailPropertyRow* SelectionFilterRow = DetailBuilder.EditDefaultProperty(SelectVerticesHandle);
+
+	const FText SelectionFilterLabel = LOCTEXT("SelectionFilterLabel", "Selection Filter");
+
+	SelectionFilterRow->CustomWidget()
+		.NameContent()
+		[
+			SelectVerticesHandle->CreatePropertyNameWidget(SelectionFilterLabel, SelectionFilterLabel)
+		]
+		.ValueContent()
+		[
+			SNew(SBox)
+			.Padding(FMargin(0, ModelingUIConstants::DetailRowVertPadding))
+			[
+				ToolbarBuilder.MakeWidget()
+			]
+		];
+}
+
+#undef LOCTEXT_NAMESPACE
\ No newline at end of file
diff --git a/Engine/Plugins/Editor/ModelingToolsEditorMode/Source/ModelingToolsEditorMode/Private/ModelingToolsEditorModeModule.cpp b/Engine/Plugins/Editor/ModelingToolsEditorMode/Source/ModelingToolsEditorMode/Private/ModelingToolsEditorModeModule.cpp
--- a/Engine/Plugins/Editor/ModelingToolsEditorMode/Source/ModelingToolsEditorMode/Private/ModelingToolsEditorModeModule.cpp
+++ b/Engine/Plugins/Editor/ModelingToolsEditorMode/Source/ModelingToolsEditorMode/Private/ModelingToolsEditorModeModule.cpp
@@ -16,8 +16,10 @@
 #include "DetailsCustomizations/MeshVertexSculptToolCustomizations.h"
 #include "DetailsCustomizations/BakeMeshAttributeToolCustomizations.h"
 #include "DetailsCustomizations/BakeTransformToolCustomizations.h"
+#include "DetailsCustomizations/MeshTopologySelectionMechanicCustomization.h"
 
 #include "PropertySets/AxisFilterPropertyType.h"
+#include "Selection/MeshTopologySelectionMechanic.h"
 #include "MeshVertexSculptTool.h"
 #include "BakeMeshAttributeMapsTool.h"
 #include "BakeMultiMeshAttributeMapsTool.h"
@@ -93,6 +95,9 @@ void FModelingToolsEditorModeModule::OnPostEngineInit()
 	ClassesToUnregisterOnShutdown.Add(UBakeMultiMeshAttributeMapsToolProperties::StaticClass()->GetFName());
 	PropertyModule.RegisterCustomClassLayout("BakeMeshAttributeVertexToolProperties", FOnGetDetailCustomizationInstance::CreateStatic(&FBakeMeshAttributeVertexToolDetails::MakeInstance));
 	ClassesToUnregisterOnShutdown.Add(UBakeMeshAttributeVertexToolProperties::StaticClass()->GetFName());
+	// PolyEd
+	PropertyModule.RegisterCustomClassLayout("MeshTopologySelectionMechanicProperties", FOnGetDetailCustomizationInstance::CreateStatic(&FMeshTopologySelectionMechanicPropertiesDetails::MakeInstance));
+	ClassesToUnregisterOnShutdown.Add(UMeshTopologySelectionMechanicProperties::StaticClass()->GetFName());
 
 	PropertyModule.RegisterCustomClassLayout("BakeTransformToolProperties", FOnGetDetailCustomizationInstance::CreateStatic(&FBakeTransformToolDetails::MakeInstance));
 	ClassesToUnregisterOnShutdown.Add(UBakeTransformToolProperties::StaticClass()->GetFName());
diff --git a/Engine/Plugins/Editor/ModelingToolsEditorMode/Source/ModelingToolsEditorMode/Private/ModelingToolsEditorModeStyle.cpp b/Engine/Plugins/Editor/ModelingToolsEditorMode/Source/ModelingToolsEditorMode/Private/ModelingToolsEditorModeStyle.cpp
--- a/Engine/Plugins/Editor/ModelingToolsEditorMode/Source/ModelingToolsEditorMode/Private/ModelingToolsEditorModeStyle.cpp
+++ b/Engine/Plugins/Editor/ModelingToolsEditorMode/Source/ModelingToolsEditorMode/Private/ModelingToolsEditorModeStyle.cpp
@@ -8,6 +8,7 @@
 #include "Interfaces/IPluginManager.h"
 #include "SlateOptMacros.h"
 #include "Styling/SlateStyleMacros.h"
+#include "Styling/ToolBarStyle.h"
 
 
 #define IMAGE_PLUGIN_BRUSH( RelativePath, ... ) FSlateImageBrush( FModelingToolsEditorModeStyle::InContent( RelativePath, ".png" ), __VA_ARGS__ )
@@ -55,6 +56,10 @@ void FModelingToolsEditorModeStyle::Initialize()
 	}
 
 	StyleSet = MakeShareable(new FSlateStyleSet(GetStyleSetName()));
+
+	// If we get asked for something that we don't set, we should default to editor style
+	StyleSet->SetParentStyleName("EditorStyle");
+
 	StyleSet->SetContentRoot(FPaths::EnginePluginsDir() / TEXT("Editor/ModelingToolsEditorMode/Content"));
 	StyleSet->SetCoreContentRoot(FPaths::EngineContentDir() / TEXT("Slate"));
 
@@ -349,8 +354,35 @@ void FModelingToolsEditorModeStyle::Initialize()
 		StyleSet->Set("BrushTypeIcons.GrabSharp", new IMAGE_BRUSH_SVG("Icons/BrushIcons/Brush_GrabSharp", Icon120));
 		StyleSet->Set("BrushTypeIcons.Twist", new IMAGE_BRUSH_SVG("Icons/BrushIcons/Brush_Twist", Icon120));
 
+		//
+		// Icons for selection buttons in PolyEd and TriEd
+		//
+
+		StyleSet->Set("PolyEd.SelectCorners", new IMAGE_BRUSH_SVG("Icons/SelectionVertices", Icon20x20));
+		StyleSet->Set("PolyEd.SelectEdges", new IMAGE_BRUSH_SVG("Icons/SelectionBorderEdges", Icon20x20));
+		StyleSet->Set("PolyEd.SelectFaces", new IMAGE_BRUSH_SVG("Icons/SelectionTriangles3", Icon20x20));
+		StyleSet->Set("PolyEd.SelectEdgeLoops", new IMAGE_BRUSH_SVG("Icons/ModelingEdgeLoopSelection", Icon20x20));
+		StyleSet->Set("PolyEd.SelectEdgeRings", new IMAGE_BRUSH_SVG("Icons/ModelingEdgeRingSelection", Icon20x20));
 	}
 
+	// Style for the toolbar in the PolyEd customization
+	{
+		// For the selection button toolbar, we want to use something similar to the toolbar we use in the viewport
+		StyleSet->Set("PolyEd.SelectionToolbar", FAppStyle::Get().GetWidgetStyle<FToolBarStyle>("EditorViewportToolBar"));
+
+		// However, increase the size of the buttons a bit
+		FCheckBoxStyle ToggleButtonStart = FAppStyle::Get().GetWidgetStyle<FCheckBoxStyle>("EditorViewportToolBar.ToggleButton.Start");
+		ToggleButtonStart.SetPadding(FMargin(9, 7, 6, 7));
+		StyleSet->Set("PolyEd.SelectionToolbar.ToggleButton.Start", ToggleButtonStart);
+
+		FCheckBoxStyle ToggleButtonMiddle = FAppStyle::Get().GetWidgetStyle<FCheckBoxStyle>("EditorViewportToolBar.ToggleButton.Middle");
+		ToggleButtonMiddle.SetPadding(FMargin(9, 7, 6, 7));
+		StyleSet->Set("PolyEd.SelectionToolbar.ToggleButton.Middle", ToggleButtonMiddle);
+
+		FCheckBoxStyle ToggleButtonEnd = FAppStyle::Get().GetWidgetStyle<FCheckBoxStyle>("EditorViewportToolBar.ToggleButton.End");
+		ToggleButtonEnd.SetPadding(FMargin(7, 7, 8, 7));
+		StyleSet->Set("PolyEd.SelectionToolbar.ToggleButton.End", ToggleButtonEnd);
+	}
 
 	FSlateStyleRegistry::RegisterSlateStyle(*StyleSet.Get());
 };
diff --git a/Engine/Plugins/Editor/ModelingToolsEditorMode/Source/ModelingToolsEditorMode/Public/DetailsCustomizations/MeshTopologySelectionMechanicCustomization.h b/Engine/Plugins/Editor/ModelingToolsEditorMode/Source/ModelingToolsEditorMode/Public/DetailsCustomizations/MeshTopologySelectionMechanicCustomization.h
new file mode 100644
--- /dev/null
+++ b/Engine/Plugins/Editor/ModelingToolsEditorMode/Source/ModelingToolsEditorMode/Public/DetailsCustomizations/MeshTopologySelectionMechanicCustomization.h
@@ -0,0 +1,16 @@
+// Copyright Epic Games, Inc. All Rights Reserved.
+
+#pragma once
+
+#include "IDetailCustomization.h"
+#include "Templates/SharedPointer.h"
+
+class IDetailLayoutBuilder;
+
+// Customization for UMeshTopologySelectionMechanicProperties
+class FMeshTopologySelectionMechanicPropertiesDetails : public IDetailCustomization
+{
+public:
+	static TSharedRef<IDetailCustomization> MakeInstance();
+	void CustomizeDetails(IDetailLayoutBuilder& DetailBuilder) override;
+};
diff --git a/Engine/Plugins/Runtime/MeshModelingToolset/Source/MeshModelingTools/Private/EditMeshPolygonsTool.cpp b/Engine/Plugins/Runtime/MeshModelingToolset/Source/MeshModelingTools/Private/EditMeshPolygonsTool.cpp
--- a/Engine/Plugins/Runtime/MeshModelingToolset/Source/MeshModelingTools/Private/EditMeshPolygonsTool.cpp
+++ b/Engine/Plugins/Runtime/MeshModelingToolset/Source/MeshModelingTools/Private/EditMeshPolygonsTool.cpp
@@ -455,6 +455,7 @@ void UEditMeshPolygonsTool::Setup()
 	SelectionMechanic->Setup(this);
 	SelectionMechanic->SetShowSelectableCorners(CommonProps->bShowSelectableCorners);
 	SelectionMechanic->Properties->RestoreProperties(this, GetPropertyCacheIdentifier(bTriangleMode));
+	SelectionMechanic->Properties->bDisplayPolygroupReliantControls = !bTriangleMode;
 	SelectionMechanic->OnSelectionChanged.AddUObject(this, &UEditMeshPolygonsTool::OnSelectionModifiedEvent);
 	SelectionMechanic->OnFaceSelectionPreviewChanged.AddWeakLambda(this, [this]() {
 		Preview->PreviewMesh->FastNotifySecondaryTrianglesChanged();
diff --git a/Engine/Plugins/Runtime/MeshModelingToolset/Source/ModelingComponents/Private/Selection/MeshTopologySelectionMechanic.cpp b/Engine/Plugins/Runtime/MeshModelingToolset/Source/ModelingComponents/Private/Selection/MeshTopologySelectionMechanic.cpp
--- a/Engine/Plugins/Runtime/MeshModelingToolset/Source/ModelingComponents/Private/Selection/MeshTopologySelectionMechanic.cpp
+++ b/Engine/Plugins/Runtime/MeshModelingToolset/Source/ModelingComponents/Private/Selection/MeshTopologySelectionMechanic.cpp
@@ -77,6 +77,12 @@ void UMeshTopologySelectionMechanic::Setup(UInteractiveTool* ParentToolIn)
 	Properties->WatchProperty(Properties->bSelectFaces, [this](bool bSelectFaces) {
 		UpdateMarqueeEnabled();
 	});
+	Properties->WatchProperty(Properties->bSelectEdgeLoops, [this](bool bSelectEdgeLoops) {
+		UpdateMarqueeEnabled();
+	});
+	Properties->WatchProperty(Properties->bSelectEdgeRings, [this](bool bSelectEdgeRings) {
+		UpdateMarqueeEnabled();
+	});
 	Properties->WatchProperty(Properties->bEnableMarquee, [this](bool bEnableMarquee) { 
 		UpdateMarqueeEnabled(); 
 	});
@@ -291,7 +297,7 @@ FGroupTopologySelector::FSelectionSettings UMeshTopologySelectionMechanic::GetTo
 	FGroupTopologySelector::FSelectionSettings Settings;
 
 	Settings.bEnableFaceHits = Properties->bSelectFaces;
-	Settings.bEnableEdgeHits = Properties->bSelectEdges;
+	Settings.bEnableEdgeHits = Properties->bSelectEdges || Properties->bSelectEdgeLoops || Properties->bSelectEdgeRings;
 	Settings.bEnableCornerHits = Properties->bSelectVertices;
 	Settings.bHitBackFaces = Properties->bHitBackFaces;
 
@@ -415,7 +421,7 @@ void UMeshTopologySelectionMechanic::SelectAll()
 	{
 		SelectAllIndices(TopologyProvider->GetNumCorners(), PersistentSelection.SelectedCornerIDs);
 	}
-	else if (Properties->bSelectEdges)
+	else if (Properties->bSelectEdges || Properties->bSelectEdgeLoops || Properties->bSelectEdgeRings)
 	{
 		SelectAllIndices(TopologyProvider->GetNumEdges(), PersistentSelection.SelectedEdgeIDs);
 	}
@@ -588,7 +594,8 @@ void UMeshTopologySelectionMechanic::UpdateMarqueeEnabled()
 	MarqueeMechanic->SetIsEnabled(
 		bIsEnabled 
 		&& Properties->bEnableMarquee
-		&& (Properties->bSelectVertices || Properties->bSelectEdges || Properties->bSelectFaces));
+		&& (Properties->bSelectVertices || Properties->bSelectEdges || Properties->bSelectFaces
+			|| Properties->bSelectEdgeLoops || Properties->bSelectEdgeRings));
 }
 
 
diff --git a/Engine/Plugins/Runtime/MeshModelingToolset/Source/ModelingComponents/Public/Selection/MeshTopologySelectionMechanic.h b/Engine/Plugins/Runtime/MeshModelingToolset/Source/ModelingComponents/Public/Selection/MeshTopologySelectionMechanic.h
--- a/Engine/Plugins/Runtime/MeshModelingToolset/Source/ModelingComponents/Public/Selection/MeshTopologySelectionMechanic.h
+++ b/Engine/Plugins/Runtime/MeshModelingToolset/Source/ModelingComponents/Public/Selection/MeshTopologySelectionMechanic.h
@@ -41,22 +41,22 @@ public:
 	bool bSelectFaces = true;
 
 	/** When true, will select edge loops. Edge loops are paths along a string of valence-4 vertices. */
-	UPROPERTY(EditAnywhere, Category = SelectionFilter, meta = (EditCondition = "bSelectEdges"))
+	UPROPERTY(EditAnywhere, Category = SelectionFilter)
 	bool bSelectEdgeLoops = false;
 
-	/** When true, will select rings of edges that are opposite each other across a quad face. */
-	UPROPERTY(EditAnywhere, Category = SelectionFilter, meta = (EditCondition = "bSelectEdges"))
+	/** When set, will select rings of edges that are opposite each other across a quad face. */
+	UPROPERTY(EditAnywhere, Category = SelectionFilter)
 	bool bSelectEdgeRings = false;
 
-	/** When false, faces that face away from the camera are ignored in selection and occlusion. Useful for working with inside-out meshes. */
-	UPROPERTY(EditAnywhere, Category = SelectionFilter, AdvancedDisplay)
+	/** When set, faces that face away from the camera are ignored in selection and occlusion. Useful for working with inside-out meshes. */
+	UPROPERTY(EditAnywhere, Category = AdditionalSelectionOptions, AdvancedDisplay)
 	bool bHitBackFaces = true;
 
-	UPROPERTY(EditAnywhere, Category = SelectionFilter, AdvancedDisplay)
+	UPROPERTY(EditAnywhere, Category = AdditionalSelectionOptions, AdvancedDisplay)
 	bool bEnableMarquee = true;
 
 	/** Determines whether vertices should be checked for occlusion in marquee select (Note: marquee select currently only works with edges and vertices) */
-	UPROPERTY(EditAnywhere, Category = SelectionFilter, meta = (EditCondition = "bEnableMarquee", EditConditionHides))
+	UPROPERTY(EditAnywhere, Category = AdditionalSelectionOptions, meta = (EditCondition = "bEnableMarquee", EditConditionHides))
 	bool bMarqueeIgnoreOcclusion = true;
 
 	// The following were originally in their own category, all marked as AdvancedDisplay. However, since there wasn't a non-AdvancedDisplay
@@ -64,17 +64,20 @@ public:
 	// The alternative approach, used below, is to have them in a nested category, which starts out as collapsed. This works nicely.
 
 	/** Prefer to select an edge projected to a point rather than the point, or a face projected to an edge rather than the edge. */
-	UPROPERTY(EditAnywhere, Category = "SelectionFilter|Ortho Viewport Behavior")
+	UPROPERTY(EditAnywhere, Category = "AdditionalSelectionOptions|Ortho Viewport Behavior")
 	bool bPreferProjectedElement = true;
 
 	/** If the closest element is valid, select other elements behind it that are aligned with it. */
-	UPROPERTY(EditAnywhere, Category = "SelectionFilter|Ortho Viewport Behavior")
+	UPROPERTY(EditAnywhere, Category = "AdditionalSelectionOptions|Ortho Viewport Behavior")
 	bool bSelectDownRay = true;
 
 	/** Do not check whether the closest element is occluded from the current view. */
-	UPROPERTY(EditAnywhere, Category = "SelectionFilter|Ortho Viewport Behavior")
+	UPROPERTY(EditAnywhere, Category = "AdditionalSelectionOptions|Ortho Viewport Behavior")
 	bool bIgnoreOcclusion = false;
 
+	// Used to avoid showing some of the selection filter buttons in triedit (in the detail customization)
+	bool bDisplayPolygroupReliantControls = true;
+
 	/** Invert current selection. If selection is empty, has same effect as Select All, and is similarly dependent on selection filter. */
 	UFUNCTION(CallInEditor, Category = SelectionActions)
 	void InvertSelection();
