MeshModelingToolset: Hide BSPConv tool in Restrictive Mode

diff --git a/Engine/Plugins/Editor/ModelingToolsEditorMode/Source/ModelingToolsEditorMode/Private/ModelingToolsEditorMode.cpp b/Engine/Plugins/Editor/ModelingToolsEditorMode/Source/ModelingToolsEditorMode/Private/ModelingToolsEditorMode.cpp
--- a/Engine/Plugins/Editor/ModelingToolsEditorMode/Source/ModelingToolsEditorMode/Private/ModelingToolsEditorMode.cpp
+++ b/Engine/Plugins/Editor/ModelingToolsEditorMode/Source/ModelingToolsEditorMode/Private/ModelingToolsEditorMode.cpp
@@ -626,7 +626,11 @@ void UModelingToolsEditorMode::Enter()
 	TrimMeshesToolBuilder->bTrimMode = true;
 	RegisterTool(ToolManagerCommands.BeginMeshTrimTool, TEXT("BeginMeshTrimTool"), TrimMeshesToolBuilder);
 
-	RegisterTool(ToolManagerCommands.BeginBspConversionTool, TEXT("BeginBspConversionTool"), NewObject<UBspConversionToolBuilder>());
+	// BSPConv is disabled in Restrictive Mode.
+	if (ToolManagerCommands.BeginBspConversionTool)
+	{
+		RegisterTool(ToolManagerCommands.BeginBspConversionTool, TEXT("BeginBspConversionTool"), NewObject<UBspConversionToolBuilder>());
+	}
 
 	auto MeshToVolumeToolBuilder = NewObject<UMeshToVolumeToolBuilder>();
 	RegisterTool(ToolManagerCommands.BeginMeshToVolumeTool, TEXT("BeginMeshToVolumeTool"), MeshToVolumeToolBuilder);
diff --git a/Engine/Plugins/Editor/ModelingToolsEditorMode/Source/ModelingToolsEditorMode/Private/ModelingToolsEditorModeToolkit.cpp b/Engine/Plugins/Editor/ModelingToolsEditorMode/Source/ModelingToolsEditorMode/Private/ModelingToolsEditorModeToolkit.cpp
--- a/Engine/Plugins/Editor/ModelingToolsEditorMode/Source/ModelingToolsEditorMode/Private/ModelingToolsEditorModeToolkit.cpp
+++ b/Engine/Plugins/Editor/ModelingToolsEditorMode/Source/ModelingToolsEditorMode/Private/ModelingToolsEditorModeToolkit.cpp
@@ -796,8 +796,14 @@ void FModelingToolsEditorModeToolkit::BuildToolPalette(FName PaletteIndex, class
 		ToolbarBuilder.AddToolBarButton(Commands.BeginVolumeToMeshTool);
 		ToolbarBuilder.AddToolBarButton(Commands.BeginMeshToVolumeTool);
 		ToolbarBuilder.AddSeparator();
-		ToolbarBuilder.AddToolBarButton(Commands.BeginBspConversionTool);
-		ToolbarBuilder.AddSeparator();
+
+		// BSPConv is disabled in Restrictive Mode.
+		if (Commands.BeginBspConversionTool)
+		{
+			ToolbarBuilder.AddToolBarButton(Commands.BeginBspConversionTool);
+			ToolbarBuilder.AddSeparator();
+		}
+
 		ToolbarBuilder.AddToolBarButton(Commands.BeginPhysicsInspectorTool);
 		ToolbarBuilder.AddToolBarButton(Commands.BeginSetCollisionGeometryTool);
 		//ToolbarBuilder.AddToolBarButton(Commands.BeginEditCollisionGeometryTool);
diff --git a/Engine/Plugins/Editor/ModelingToolsEditorMode/Source/ModelingToolsEditorMode/Private/ModelingToolsManagerActions.cpp b/Engine/Plugins/Editor/ModelingToolsEditorMode/Source/ModelingToolsEditorMode/Private/ModelingToolsManagerActions.cpp
--- a/Engine/Plugins/Editor/ModelingToolsEditorMode/Source/ModelingToolsEditorMode/Private/ModelingToolsManagerActions.cpp
+++ b/Engine/Plugins/Editor/ModelingToolsEditorMode/Source/ModelingToolsEditorMode/Private/ModelingToolsManagerActions.cpp
@@ -3,6 +3,7 @@
 #include "ModelingToolsManagerActions.h"
 #include "InputCoreTypes.h"
 #include "ModelingToolsEditorModeStyle.h"
+#include "ModelingToolsEditorModeSettings.h"
 
 #define LOCTEXT_NAMESPACE "ModelingToolsManagerCommands"
 
@@ -38,6 +39,8 @@ TSharedPtr<FUICommandInfo> FModelingToolsManagerCommands::FindToolByName(FString
 
 void FModelingToolsManagerCommands::RegisterCommands()
 {
+	const UModelingToolsEditorModeSettings* Settings = GetDefault<UModelingToolsEditorModeSettings>();
+
 	// this has to be done with a compile-time macro because UI_COMMAND expands to LOCTEXT macros
 #define REGISTER_MODELING_TOOL_COMMAND(ToolCommandInfo, ToolName, ToolTip) \
 		UI_COMMAND(ToolCommandInfo, ToolName, ToolTip, EUserInterfaceActionType::ToggleButton, FInputChord()); \
@@ -146,7 +149,10 @@ void FModelingToolsManagerCommands::RegisterCommands()
 	// Volumes
 	REGISTER_MODELING_TOOL_COMMAND(BeginVolumeToMeshTool, "Vol2Msh", "Convert volume to new mesh asset");
 	REGISTER_MODELING_TOOL_COMMAND(BeginMeshToVolumeTool, "Msh2Vol", "Convert mesh to volume");
-	REGISTER_MODELING_TOOL_COMMAND(BeginBspConversionTool, "BSPConv", "Convert BSP to static mesh asset");
+	if (!Settings->InRestrictiveMode())
+	{
+		REGISTER_MODELING_TOOL_COMMAND(BeginBspConversionTool, "BSPConv", "Convert BSP to static mesh asset");
+	}
 	REGISTER_MODELING_TOOL_COMMAND(BeginPhysicsInspectorTool, "PInspct", "Inspect physics geometry for selected meshes");
 	REGISTER_MODELING_TOOL_COMMAND(BeginSetCollisionGeometryTool, "Msh2Coll", "Convert selected meshes to Simple Collision Geometry (for last selected)");
 	REGISTER_MODELING_TOOL_COMMAND(BeginExtractCollisionGeometryTool, "Coll2Msh", "Convert Simple Collision Geometry to new mesh asset");
