Move the ADynamicMeshActor Outliner filter from ModelingTools to GeometryFramework

diff --git a/Engine/Plugins/Editor/ModelingToolsEditorMode/Source/ModelingToolsEditorMode/ModelingToolsEditorMode.Build.cs b/Engine/Plugins/Editor/ModelingToolsEditorMode/Source/ModelingToolsEditorMode/ModelingToolsEditorMode.Build.cs
--- a/Engine/Plugins/Editor/ModelingToolsEditorMode/Source/ModelingToolsEditorMode/ModelingToolsEditorMode.Build.cs
+++ b/Engine/Plugins/Editor/ModelingToolsEditorMode/Source/ModelingToolsEditorMode/ModelingToolsEditorMode.Build.cs
@@ -49,7 +49,6 @@ public class ModelingToolsEditorMode : ModuleRules
 				"StylusInput",
 				"DeveloperSettings",
 				"PropertyEditor",
-				"GeometryFramework"
 				// ... add private dependencies that you statically link with here ...	
 			}
 		);
diff --git a/Engine/Plugins/Editor/ModelingToolsEditorMode/Source/ModelingToolsEditorMode/Private/ModelingToolsEditorModeModule.cpp b/Engine/Plugins/Editor/ModelingToolsEditorMode/Source/ModelingToolsEditorMode/Private/ModelingToolsEditorModeModule.cpp
--- a/Engine/Plugins/Editor/ModelingToolsEditorMode/Source/ModelingToolsEditorMode/Private/ModelingToolsEditorModeModule.cpp
+++ b/Engine/Plugins/Editor/ModelingToolsEditorMode/Source/ModelingToolsEditorMode/Private/ModelingToolsEditorModeModule.cpp
@@ -25,10 +25,6 @@
 #include "BakeMultiMeshAttributeMapsTool.h"
 #include "BakeMeshAttributeVertexTool.h"
 #include "BakeTransformTool.h"
-#include "DynamicMeshActor.h"
-
-#include "LevelEditor.h"
-#include "Filters/CustomClassFilterData.h"
 
 
 #define LOCTEXT_NAMESPACE "FModelingToolsEditorModeModule"
@@ -105,15 +101,6 @@ void FModelingToolsEditorModeModule::OnPostEngineInit()
 
 	PropertyModule.RegisterCustomClassLayout("BakeTransformToolProperties", FOnGetDetailCustomizationInstance::CreateStatic(&FBakeTransformToolDetails::MakeInstance));
 	ClassesToUnregisterOnShutdown.Add(UBakeTransformToolProperties::StaticClass()->GetFName());
-
-	FLevelEditorModule& LevelEditorModule = FModuleManager::LoadModuleChecked<FLevelEditorModule>("LevelEditor");
-
-	// Register Level Editor Outliner filter
-	if(TSharedPtr<FFilterCategory> GeometryFilterCategory = LevelEditorModule.GetOutlinerFilterCategory(FLevelEditorOutlinerBuiltInCategories::Geometry()))
-	{
-		TSharedRef<FCustomClassFilterData> DynamicMeshActorClassData = MakeShared<FCustomClassFilterData>(ADynamicMeshActor::StaticClass(), GeometryFilterCategory, FLinearColor::White);
-		LevelEditorModule.AddCustomClassFilterToOutliner(DynamicMeshActorClassData);
-	}
 }
 
 #undef LOCTEXT_NAMESPACE
diff --git a/Engine/Source/Runtime/GeometryFramework/GeometryFramework.Build.cs b/Engine/Source/Runtime/GeometryFramework/GeometryFramework.Build.cs
--- a/Engine/Source/Runtime/GeometryFramework/GeometryFramework.Build.cs
+++ b/Engine/Source/Runtime/GeometryFramework/GeometryFramework.Build.cs
@@ -10,6 +10,24 @@ public class GeometryFramework : ModuleRules
 
 		// this is required for raytracing suppport?
 		PublicIncludePaths.Add("../Shaders/Shared");
+		
+		// These are required to register filters to the Level Editor Outliner
+		if (Target.Type == TargetType.Editor)
+		{
+			PrivateDependencyModuleNames.AddRange(
+				new string[]
+				{
+					"EditorWidgets",
+				}
+			);
+				
+			PrivateIncludePathModuleNames.AddRange(
+				new string[]
+				{													
+					"LevelEditor",
+				}
+			);
+		}
 
 		PublicDependencyModuleNames.AddRange(
 			new string[] {
diff --git a/Engine/Source/Runtime/GeometryFramework/Private/GeometryFrameworkModule.cpp b/Engine/Source/Runtime/GeometryFramework/Private/GeometryFrameworkModule.cpp
--- a/Engine/Source/Runtime/GeometryFramework/Private/GeometryFrameworkModule.cpp
+++ b/Engine/Source/Runtime/GeometryFramework/Private/GeometryFrameworkModule.cpp
@@ -4,11 +4,33 @@
 #include "Components/BaseDynamicMeshComponent.h"
 #include "Misc/CoreDelegates.h"
 
+#if WITH_EDITOR
+
+#include "DynamicMeshActor.h"
+#include "LevelEditor.h"
+#include "Filters/CustomClassFilterData.h"
+
+#endif
+
 #define LOCTEXT_NAMESPACE "FGeometryFrameworkModule"
 
 void FGeometryFrameworkModule::StartupModule()
 {
 	FCoreDelegates::OnPostEngineInit.AddRaw(this, &FGeometryFrameworkModule::OnPostEngineInit);
+
+#if WITH_EDITOR
+	if (!IsRunningGame())
+	{
+		FLevelEditorModule& LevelEditorModule = FModuleManager::LoadModuleChecked<FLevelEditorModule>("LevelEditor");
+
+		// Register the Level Editor Outliner filter for ADynamicMeshActor
+		if(TSharedPtr<FFilterCategory> GeometryFilterCategory = LevelEditorModule.GetOutlinerFilterCategory(FLevelEditorOutlinerBuiltInCategories::Geometry()))
+		{
+			TSharedRef<FCustomClassFilterData> DynamicMeshActorClassData = MakeShared<FCustomClassFilterData>(ADynamicMeshActor::StaticClass(), GeometryFilterCategory, FLinearColor::White);
+			LevelEditorModule.AddCustomClassFilterToOutliner(DynamicMeshActorClassData);
+		}
+	}
+#endif
 }
 
 void FGeometryFrameworkModule::OnPostEngineInit()
