UVEditor: Cleanup having to do with the modeling tools button. Clarify comments, remove UV editor from IGeometryProcessingInterfacesModule interface since it shouldn't have multiple implementations and isn't a processing algorithm, rename classes that sound like they are asset editors in themselves.

diff --git a/Engine/Plugins/Editor/ModelingToolsEditorMode/Source/ModelingToolsEditorMode/Private/ModelingToolsEditorMode.cpp b/Engine/Plugins/Editor/ModelingToolsEditorMode/Source/ModelingToolsEditorMode/Private/ModelingToolsEditorMode.cpp
--- a/Engine/Plugins/Editor/ModelingToolsEditorMode/Source/ModelingToolsEditorMode/Private/ModelingToolsEditorMode.cpp
+++ b/Engine/Plugins/Editor/ModelingToolsEditorMode/Source/ModelingToolsEditorMode/Private/ModelingToolsEditorMode.cpp
@@ -2,7 +2,7 @@
 
 #include "ModelingToolsEditorMode.h"
 #include "IGeometryProcessingInterfacesModule.h"
-#include "GeometryProcessingInterfaces/UVEditorAssetEditor.h"
+#include "GeometryProcessingInterfaces/IUVEditorModularFeature.h"
 #include "InteractiveTool.h"
 #include "ModelingToolsEditorModeToolkit.h"
 #include "ModelingToolsEditorModeSettings.h"
@@ -877,13 +877,17 @@ void UModelingToolsEditorMode::Enter()
 
 void UModelingToolsEditorMode::RegisterUVEditor()
 {
-	// Handle the inclusion of the optional UVEditor button if the UVEditor plugin has been found
-	IGeometryProcessingInterfacesModule& GeomProcInterfaces = FModuleManager::Get().LoadModuleChecked<IGeometryProcessingInterfacesModule>("GeometryProcessingInterfaces");
-	IGeometryProcessing_UVEditorAssetEditor* UVEditorAPI = GeomProcInterfaces.GetUVEditorAssetEditorImplementation();
-	const FModelingToolsManagerCommands& ToolManagerCommands = FModelingToolsManagerCommands::Get();
+	IUVEditorModularFeature* UVEditorAPI = nullptr;
+	// We should be allowed to do GetModularFeatureImplementation directly without the check, but currently there is an assert
+	// there (despite what the header for that function promises).
+	if (IModularFeatures::Get().IsModularFeatureAvailable(IUVEditorModularFeature::GetModularFeatureName()))
+	{
+		UVEditorAPI = static_cast<IUVEditorModularFeature*>(IModularFeatures::Get().GetModularFeatureImplementation(IUVEditorModularFeature::GetModularFeatureName(), 0));
+	}
 
 	if (UVEditorAPI)
 	{
+		const FModelingToolsManagerCommands& ToolManagerCommands = FModelingToolsManagerCommands::Get();
 		const TSharedRef<FUICommandList>& CommandList = Toolkit->GetToolkitCommands();
 		CommandList->MapAction(ToolManagerCommands.LaunchUVEditor,
 			FExecuteAction::CreateLambda([this, UVEditorAPI]() 
diff --git a/Engine/Plugins/Editor/ModelingToolsEditorMode/Source/ModelingToolsEditorMode/Private/ModelingToolsEditorModeToolkit.cpp b/Engine/Plugins/Editor/ModelingToolsEditorMode/Source/ModelingToolsEditorMode/Private/ModelingToolsEditorModeToolkit.cpp
--- a/Engine/Plugins/Editor/ModelingToolsEditorMode/Source/ModelingToolsEditorMode/Private/ModelingToolsEditorModeToolkit.cpp
+++ b/Engine/Plugins/Editor/ModelingToolsEditorMode/Source/ModelingToolsEditorMode/Private/ModelingToolsEditorModeToolkit.cpp
@@ -2,7 +2,7 @@
 
 #include "ModelingToolsEditorModeToolkit.h"
 #include "IGeometryProcessingInterfacesModule.h"
-#include "GeometryProcessingInterfaces/UVEditorAssetEditor.h"
+#include "GeometryProcessingInterfaces/IUVEditorModularFeature.h"
 #include "ModelingToolsEditorMode.h"
 #include "ModelingToolsManagerActions.h"
 #include "ModelingToolsEditorModeSettings.h"
@@ -783,10 +783,9 @@ void FModelingToolsEditorModeToolkit::BuildToolPalette(FName PaletteIndex, class
 		ToolbarBuilder.AddToolBarButton(Commands.BeginUVSeamEditTool);
 		ToolbarBuilder.AddToolBarButton(Commands.BeginTransformUVIslandsTool);
 		ToolbarBuilder.AddToolBarButton(Commands.BeginUVLayoutTool);
+
 		// Handle the inclusion of the optional UVEditor button if the UVEditor plugin has been found
-		IGeometryProcessingInterfacesModule& GeomProcInterfaces = FModuleManager::Get().LoadModuleChecked<IGeometryProcessingInterfacesModule>("GeometryProcessingInterfaces");
-		IGeometryProcessing_UVEditorAssetEditor* UVEditorAPI = GeomProcInterfaces.GetUVEditorAssetEditorImplementation();
-		if (UVEditorAPI)
+		if (IModularFeatures::Get().IsModularFeatureAvailable(IUVEditorModularFeature::GetModularFeatureName()))
 		{
 			ToolbarBuilder.AddToolBarButton(Commands.LaunchUVEditor);
 		}
diff --git a/Engine/Plugins/Experimental/UVEditor/Source/UVEditor/Private/UVEditorAssetEditorImpl.cpp b/Engine/Plugins/Experimental/UVEditor/Source/UVEditor/Private/UVEditorAssetEditorImpl.cpp
deleted file mode 100644
diff --git a/Engine/Plugins/Experimental/UVEditor/Source/UVEditor/Private/UVEditorModularFeature.cpp b/Engine/Plugins/Experimental/UVEditor/Source/UVEditor/Private/UVEditorModularFeature.cpp
new file mode 100644
--- /dev/null
+++ b/Engine/Plugins/Experimental/UVEditor/Source/UVEditor/Private/UVEditorModularFeature.cpp
@@ -0,0 +1,63 @@
+// Copyright Epic Games, Inc. All Rights Reserved.
+
+#include "UVEditorModularFeature.h"
+
+#include "CoreMinimal.h"
+#include "UVEditorSubsystem.h"
+#include "Editor.h"
+#include "GameFramework/Actor.h"
+
+namespace UE
+{
+namespace Geometry
+{
+	void FUVEditorModularFeature::LaunchUVEditor(const TArray<TObjectPtr<UObject>>& ObjectsIn)
+	{
+		UUVEditorSubsystem* UVSubsystem = GEditor->GetEditorSubsystem<UUVEditorSubsystem>();
+		check(UVSubsystem);
+		TArray<TObjectPtr<UObject>> ProcessedObjects;
+		ConvertInputArgsToValidTargets(ObjectsIn, ProcessedObjects);
+		UVSubsystem->StartUVEditor(ProcessedObjects);
+	}
+
+	bool FUVEditorModularFeature::CanLaunchUVEditor(const TArray<TObjectPtr<UObject>>& ObjectsIn)
+	{
+		UUVEditorSubsystem* UVSubsystem = GEditor->GetEditorSubsystem<UUVEditorSubsystem>();
+		check(UVSubsystem);
+		TArray<TObjectPtr<UObject>> ProcessedObjects;
+		ConvertInputArgsToValidTargets(ObjectsIn, ProcessedObjects);
+		return UVSubsystem->AreObjectsValidTargets(ProcessedObjects);
+	}
+
+	void FUVEditorModularFeature::ConvertInputArgsToValidTargets(const TArray<TObjectPtr<UObject>>& ObjectsIn, TArray<TObjectPtr<UObject>>& ObjectsOut) const
+	{
+		for (const TObjectPtr<UObject>& Object : ObjectsIn)
+		{
+			const AActor* Actor = Cast<const AActor>(Object);
+			if (Actor)
+			{
+				TArray<UObject*> ActorAssets;
+				Actor->GetReferencedContentObjects(ActorAssets);
+
+				if (ActorAssets.Num() > 0)
+				{
+					for (UObject* Asset : ActorAssets)
+					{
+						ObjectsOut.AddUnique(Asset);
+					}
+				}
+				else {
+					// Need to transform actors to components here because that's what the UVEditor expects to have
+					TInlineComponentArray<UActorComponent*> ActorComponents;
+					Actor->GetComponents(ActorComponents);
+					ObjectsOut.Append(ActorComponents);
+				}
+			}
+			else
+			{
+				ObjectsOut.AddUnique(Object);
+			}
+		}
+	}
+}
+}
diff --git a/Engine/Plugins/Experimental/UVEditor/Source/UVEditor/Private/UVEditorModule.cpp b/Engine/Plugins/Experimental/UVEditor/Source/UVEditor/Private/UVEditorModule.cpp
--- a/Engine/Plugins/Experimental/UVEditor/Source/UVEditor/Private/UVEditorModule.cpp
+++ b/Engine/Plugins/Experimental/UVEditor/Source/UVEditor/Private/UVEditorModule.cpp
@@ -21,11 +21,11 @@
 
 void FUVEditorModule::StartupModule()
 {
-	UVEditorAssetEditor = MakeShared<UE::Geometry::FUVEditorAssetEditorImpl>();
+	UVEditorAssetEditor = MakeShared<UE::Geometry::FUVEditorModularFeature>();
 
 	if (UVEditorAssetEditor.IsValid())
 	{
-		IModularFeatures::Get().RegisterModularFeature(IGeometryProcessing_UVEditorAssetEditor::GetModularFeatureName(), UVEditorAssetEditor.Get());
+		IModularFeatures::Get().RegisterModularFeature(IUVEditorModularFeature::GetModularFeatureName(), UVEditorAssetEditor.Get());
 	}
 
 	FUVEditorStyle::Get(); // Causes the constructor to be called
@@ -48,7 +48,7 @@ void FUVEditorModule::ShutdownModule()
 {
 	if (UVEditorAssetEditor.IsValid())
 	{
-		IModularFeatures::Get().UnregisterModularFeature(IGeometryProcessing_UVEditorAssetEditor::GetModularFeatureName(), UVEditorAssetEditor.Get());
+		IModularFeatures::Get().UnregisterModularFeature(IUVEditorModularFeature::GetModularFeatureName(), UVEditorAssetEditor.Get());
 		UVEditorAssetEditor = nullptr;
 	}
 
diff --git a/Engine/Plugins/Experimental/UVEditor/Source/UVEditor/Public/UVEditorAssetEditorImpl.h b/Engine/Plugins/Experimental/UVEditor/Source/UVEditor/Public/UVEditorAssetEditorImpl.h
deleted file mode 100644
diff --git a/Engine/Plugins/Experimental/UVEditor/Source/UVEditor/Public/UVEditorModularFeature.h b/Engine/Plugins/Experimental/UVEditor/Source/UVEditor/Public/UVEditorModularFeature.h
new file mode 100644
--- /dev/null
+++ b/Engine/Plugins/Experimental/UVEditor/Source/UVEditor/Public/UVEditorModularFeature.h
@@ -0,0 +1,30 @@
+// Copyright Epic Games, Inc. All Rights Reserved.
+
+#pragma once
+
+#include "CoreMinimal.h"
+#include "GeometryProcessingInterfaces/IUVEditorModularFeature.h"
+
+namespace UE
+{
+namespace Geometry
+{
+
+/** 
+ * Connector class that allows other plugins to look for the UV Editor to see if the plugin is present,
+ * (via IModularFeatures::Get().GetModularFeature and related methods), and then launch it if so.
+ */
+class UVEDITOR_API FUVEditorModularFeature : public IUVEditorModularFeature
+{
+public:
+	virtual void LaunchUVEditor(const TArray<TObjectPtr<UObject>>& Objects) override;
+	virtual bool CanLaunchUVEditor(const TArray<TObjectPtr<UObject>>& Objects) override;
+
+protected:
+	void ConvertInputArgsToValidTargets(const TArray<TObjectPtr<UObject>>& ObjectsIn, TArray<TObjectPtr<UObject>>& ObjectsOut) const;
+
+};
+
+
+}
+}
diff --git a/Engine/Plugins/Experimental/UVEditor/Source/UVEditor/Public/UVEditorModule.h b/Engine/Plugins/Experimental/UVEditor/Source/UVEditor/Public/UVEditorModule.h
--- a/Engine/Plugins/Experimental/UVEditor/Source/UVEditor/Public/UVEditorModule.h
+++ b/Engine/Plugins/Experimental/UVEditor/Source/UVEditor/Public/UVEditorModule.h
@@ -5,7 +5,7 @@
 #include "CoreMinimal.h"
 #include "Modules/ModuleInterface.h"
 
-#include "UVEditorAssetEditorImpl.h"
+#include "UVEditorModularFeature.h"
 
 class FLayoutExtender;
 
@@ -33,5 +33,5 @@ private:
 	/** StaticClass is not safe on shutdown, so we cache the name, and use this to unregister on shut down */
 	TArray<FName> ClassesToUnregisterOnShutdown;
 
-	TSharedPtr<UE::Geometry::FUVEditorAssetEditorImpl> UVEditorAssetEditor;
+	TSharedPtr<UE::Geometry::FUVEditorModularFeature> UVEditorAssetEditor;
 };
diff --git a/Engine/Source/Developer/GeometryProcessingInterfaces/Private/GeometryProcessingInterfacesModule.cpp b/Engine/Source/Developer/GeometryProcessingInterfaces/Private/GeometryProcessingInterfacesModule.cpp
--- a/Engine/Source/Developer/GeometryProcessingInterfaces/Private/GeometryProcessingInterfacesModule.cpp
+++ b/Engine/Source/Developer/GeometryProcessingInterfaces/Private/GeometryProcessingInterfacesModule.cpp
@@ -6,7 +6,6 @@
 #include "Modules/ModuleManager.h"
 
 #include "GeometryProcessingInterfaces/ApproximateActors.h"
-#include "GeometryProcessingInterfaces/UVEditorAssetEditor.h"
 
 
 //DEFINE_LOG_CATEGORY_STATIC(LogMeshReduction, Verbose, All);
@@ -38,17 +37,3 @@ IGeometryProcessing_ApproximateActors* FGeometryProcessingInterfacesModule::GetA
 
 	return ApproximateActors;
 }
-
-IGeometryProcessing_UVEditorAssetEditor* FGeometryProcessingInterfacesModule::GetUVEditorAssetEditorImplementation()
-{
-	if (UVEditorAssetEditor == nullptr)
-	{
-		TArray<IGeometryProcessing_UVEditorAssetEditor*> UVEditorAssetEditorOptions =
-			IModularFeatures::Get().GetModularFeatureImplementations<IGeometryProcessing_UVEditorAssetEditor>(IGeometryProcessing_UVEditorAssetEditor::GetModularFeatureName());
-
-		UVEditorAssetEditor = (UVEditorAssetEditorOptions.Num() > 0) ? UVEditorAssetEditorOptions[0] : nullptr;
-	}
-
-	return UVEditorAssetEditor;
-}
-
diff --git a/Engine/Source/Developer/GeometryProcessingInterfaces/Private/GeometryProcessingInterfacesModule.h b/Engine/Source/Developer/GeometryProcessingInterfaces/Private/GeometryProcessingInterfacesModule.h
--- a/Engine/Source/Developer/GeometryProcessingInterfaces/Private/GeometryProcessingInterfacesModule.h
+++ b/Engine/Source/Developer/GeometryProcessingInterfaces/Private/GeometryProcessingInterfacesModule.h
@@ -6,8 +6,6 @@
 
 
 class IGeometryProcessing_ApproximateActors;
-class IGeometryProcessing_UVEditorAssetEditor;
-
 
 /**
  * Implementation of IGeometryProcessingInterfacesModule (which extends the standard IModuleInterface)
@@ -26,13 +24,6 @@ public:
 	 */
 	virtual IGeometryProcessing_ApproximateActors* GetApproximateActorsImplementation() override;
 
-
-	/**
-	* @return implementation of IGeometryProcessing_UVEditor, if available, or nullptr (result is cached internally)
-	*/
-	virtual IGeometryProcessing_UVEditorAssetEditor* GetUVEditorAssetEditorImplementation() override;
-
 private:
 	IGeometryProcessing_ApproximateActors* ApproximateActors = nullptr;
-	IGeometryProcessing_UVEditorAssetEditor* UVEditorAssetEditor = nullptr;
 };
diff --git a/Engine/Source/Developer/GeometryProcessingInterfaces/Public/GeometryProcessingInterfaces/IUVEditorModularFeature.h b/Engine/Source/Developer/GeometryProcessingInterfaces/Public/GeometryProcessingInterfaces/IUVEditorModularFeature.h
new file mode 100644
--- /dev/null
+++ b/Engine/Source/Developer/GeometryProcessingInterfaces/Public/GeometryProcessingInterfaces/IUVEditorModularFeature.h
@@ -0,0 +1,33 @@
+// Copyright Epic Games, Inc. All Rights Reserved.
+
+#pragma once
+
+#include "CoreMinimal.h"
+#include "Features/IModularFeature.h"
+
+/**
+ * This interface acts as a connector to the (optional) UV Editor plugin. If the plugin is available,
+ * then it can be retrieved via IModularFeatures::Get().GetModularFeatureImplementation or related methods.
+ */
+class IUVEditorModularFeature : public IModularFeature
+{
+public:
+	virtual ~IUVEditorModularFeature() {}
+
+	virtual void LaunchUVEditor(const TArray<TObjectPtr<UObject>>& Objects)
+	{
+		check(false);		// not implemented in base class
+	}
+
+	virtual bool CanLaunchUVEditor(const TArray<TObjectPtr<UObject>>& Objects)
+	{
+		check(false);		// not implemented in base class
+		return false;
+	}
+
+	// Modular feature name to register for retrieval during runtime
+	static const FName GetModularFeatureName()
+	{
+		return TEXT("UVEditorModularFeature");
+	}
+};
\ No newline at end of file
diff --git a/Engine/Source/Developer/GeometryProcessingInterfaces/Public/GeometryProcessingInterfaces/UVEditorAssetEditor.h b/Engine/Source/Developer/GeometryProcessingInterfaces/Public/GeometryProcessingInterfaces/UVEditorAssetEditor.h
deleted file mode 100644
diff --git a/Engine/Source/Developer/GeometryProcessingInterfaces/Public/IGeometryProcessingInterfacesModule.h b/Engine/Source/Developer/GeometryProcessingInterfaces/Public/IGeometryProcessingInterfacesModule.h
--- a/Engine/Source/Developer/GeometryProcessingInterfaces/Public/IGeometryProcessingInterfacesModule.h
+++ b/Engine/Source/Developer/GeometryProcessingInterfaces/Public/IGeometryProcessingInterfacesModule.h
@@ -5,7 +5,6 @@
 #include "Modules/ModuleInterface.h"
 
 class IGeometryProcessing_ApproximateActors;
-class IGeometryProcessing_UVEditorAssetEditor;
 
 /**
  * Abstract interface to a Module that provides functions to access
@@ -29,9 +28,4 @@ public:
 	 */
 	virtual IGeometryProcessing_ApproximateActors* GetApproximateActorsImplementation() = 0;
 
-	/**
-	 * @return implementation of IGeometryProcessing_UVEditorAssetEditor, if available
-	*/
-	virtual IGeometryProcessing_UVEditorAssetEditor* GetUVEditorAssetEditorImplementation() = 0;
-
 };
