Modeling Mode: evolution of Mesh Selection system.
UModelingSelectionInteraction now supports concept of a "drag mode" that allows for drawing a marquee, brush-style interaction, etc. Added support for Marquee drawing but actual selection based on the marquee is not implemented yet.
ModelingToolsEditorMode now registers selector factories for Volumes and Static Meshes. Creating selectors for these types is dynamically configurable via new UProperties.
Added mesh-selection-specific UPathSelectionInteraction which provides a simple continuous-select-while-dragging interaction
Added a viewport overlay toolbar for mesh selection, supports changing selection mode, drag mode, selectable object types.
Viewport overlay toolbar also exposes a "quick settings" flyout for changing various modeling mode settings
Added URectangleMarqueeInteraction which is a simplified variant of URectangleMarqueeMechanic that does not require a Tool and doesn't create it's own InputBehavior
USingleClickOrDragInputBehavior now supports changing Drag target
Added a few placeholder icons

diff --git a/Engine/Plugins/Editor/ModelingToolsEditorMode/Source/ModelingToolsEditorMode/Private/ModelingModeToolkit_Toolbars.cpp b/Engine/Plugins/Editor/ModelingToolsEditorMode/Source/ModelingToolsEditorMode/Private/ModelingModeToolkit_Toolbars.cpp
new file mode 100644
--- /dev/null
+++ b/Engine/Plugins/Editor/ModelingToolsEditorMode/Source/ModelingToolsEditorMode/Private/ModelingModeToolkit_Toolbars.cpp
@@ -0,0 +1,455 @@
+// Copyright Epic Games, Inc. All Rights Reserved.
+
+#include "ModelingToolsEditorModeToolkit.h"
+#include "ModelingToolsEditorMode.h"
+#include "ModelingToolsManagerActions.h"
+#include "ModelingToolsEditorModeSettings.h"
+#include "ModelingToolsEditorModeStyle.h"
+
+#include "Selection/GeometrySelectionManager.h"
+#include "ModelingSelectionInteraction.h"
+
+#include "ModelingComponentsSettings.h"
+#include "PropertySets/CreateMeshObjectTypeProperties.h"
+
+#include "Modules/ModuleManager.h"
+#include "ISettingsModule.h"
+
+#include "Framework/MultiBox/MultiBoxBuilder.h"
+#include "Framework/Commands/UIAction.h"
+#include "Widgets/Images/SImage.h"
+#include "Widgets/Layout/SSpacer.h"
+#include "Widgets/Input/SComboButton.h"
+
+// to toggle component instance selection
+#include "ComponentRecreateRenderStateContext.h"
+#include "Elements/SMInstance/SMInstanceElementData.h"
+
+#define LOCTEXT_NAMESPACE "FModelingToolsEditorModeToolkit"
+
+
+namespace UELocal
+{
+
+void MakeSubMenu_QuickSettings(FMenuBuilder& MenuBuilder)
+{
+	const FUIAction OpenModelingModeProjectSettings(
+		FExecuteAction::CreateLambda([]
+		{
+			if (ISettingsModule* SettingsModule = FModuleManager::GetModulePtr<ISettingsModule>("Settings"))
+			{
+				SettingsModule->ShowViewer("Project", "Plugins", "ModelingMode");
+			}
+		}), FCanExecuteAction(), FIsActionChecked());
+	MenuBuilder.AddMenuEntry(LOCTEXT("ModelingModeProjectSettings", "Modeling Mode (Project)"), 
+		LOCTEXT("ModelingModeProjectSettings_Tooltip", "Jump to the Project Settings for Modeling Mode. Project Settings are Poject-specific."),
+		FSlateIcon(), OpenModelingModeProjectSettings, NAME_None, EUserInterfaceActionType::Button);
+
+	const FUIAction OpenModelingModeEditorSettings(
+		FExecuteAction::CreateLambda([]
+		{
+			if (ISettingsModule* SettingsModule = FModuleManager::GetModulePtr<ISettingsModule>("Settings"))
+			{
+				SettingsModule->ShowViewer("Editor", "Plugins", "ModelingMode");
+			}
+		}), FCanExecuteAction(), FIsActionChecked());
+	MenuBuilder.AddMenuEntry(LOCTEXT("ModelingModeEditorSettings", "Modeling Mode (Editor)"), 
+		LOCTEXT("ModelingModeEditorSettings_Tooltip", "Jump to the Editor Settings for Modeling Mode. Editor Settings apply across all Projects."),
+		FSlateIcon(), OpenModelingModeEditorSettings, NAME_None, EUserInterfaceActionType::Button);
+
+	const FUIAction OpenModelingToolsProjectSettings(
+		FExecuteAction::CreateLambda([]
+		{
+			if (ISettingsModule* SettingsModule = FModuleManager::GetModulePtr<ISettingsModule>("Settings"))
+			{
+				SettingsModule->ShowViewer("Project", "Plugins", "Modeling Mode Tools");
+			}
+		}), FCanExecuteAction(), FIsActionChecked());
+	MenuBuilder.AddMenuEntry(LOCTEXT("ModelingToolsProjectSettings", "Modeling Tools (Project)"), 
+		LOCTEXT("ModelingToolsProjectSettings_Tooltip", "Jump to the Project Settings for Modeling Tools. Project Settings are Poject-specific."),
+		FSlateIcon(), OpenModelingToolsProjectSettings, NAME_None, EUserInterfaceActionType::Button);
+
+}
+
+void MakeSubMenu_GizmoVisibilityMode(FModelingToolsEditorModeToolkit* Toolkit, FMenuBuilder& MenuBuilder)
+{
+	UModelingToolsEditorModeSettings* Settings = GetMutableDefault<UModelingToolsEditorModeSettings>();
+	UEditorInteractiveToolsContext* Context = Toolkit->GetScriptableEditorMode()->GetInteractiveToolsContext(EToolsContextScope::EdMode);
+
+	const FUIAction GizmoMode_Combined(
+		FExecuteAction::CreateLambda([Settings, Context]
+		{
+			Settings->bRespectLevelEditorGizmoMode = ! Settings->bRespectLevelEditorGizmoMode;
+			Context->SetForceCombinedGizmoMode(Settings->bRespectLevelEditorGizmoMode == false);
+			Settings->SaveConfig();
+		}),
+		FCanExecuteAction(),
+		FIsActionChecked::CreateLambda([Settings]()
+		{
+			return Settings->bRespectLevelEditorGizmoMode == false;
+		}));
+	MenuBuilder.AddMenuEntry(LOCTEXT("GizmoMode_Combined", "Combined Gizmo"), 
+		LOCTEXT("GizmoMode_Combined_Tooltip", "Ignore Level Editor Gizmo Mode and always use a Combined Transform Gizmo in Modeling Tools"),
+		FSlateIcon(), GizmoMode_Combined, NAME_None, EUserInterfaceActionType::ToggleButton);
+}
+
+
+void MakeSubMenu_ModeToggles(FMenuBuilder& MenuBuilder)
+{
+	const FUIAction ToggleInstancesEverywhereAction(
+		FExecuteAction::CreateLambda([]
+		{
+			IConsoleVariable* CVar = IConsoleManager::Get().FindConsoleVariable(TEXT("TypedElements.EnableViewportSMInstanceSelection")); 
+			int32 CurValue = CVar->GetInt();
+			CVar->Set(CurValue == 0 ? 1 : 0);
+		}),
+		FCanExecuteAction(),
+		FIsActionChecked::CreateLambda([]()
+		{
+			IConsoleVariable* CVar = IConsoleManager::Get().FindConsoleVariable(TEXT("TypedElements.EnableViewportSMInstanceSelection"));
+			return (CVar->GetInt() == 1);
+		}));
+
+	MenuBuilder.AddMenuEntry(
+		LOCTEXT("ToggleInstancesSelection", "Select Instances"), 
+		LOCTEXT("ToggleInstancesSelection_Tooltip", "Enable/Disable support for direct selection of InstancedStaticMeshComponent Instances (via TypedElements.EnableViewportInstanceSelection cvar)"),
+		FSlateIcon(), ToggleInstancesEverywhereAction, NAME_None, EUserInterfaceActionType::ToggleButton);
+
+}
+
+
+void MakeSubMenu_DefaultMeshObjectPhysicsSettings(FMenuBuilder& MenuBuilder)
+{
+	UModelingComponentsSettings* Settings = GetMutableDefault<UModelingComponentsSettings>();
+
+	const FUIAction ToggleEnableCollisionAction = FUIAction(
+		FExecuteAction::CreateLambda([Settings]
+		{
+			Settings->bEnableCollision = !Settings->bEnableCollision;
+			Settings->SaveConfig();
+		}),
+		FCanExecuteAction(),
+		FIsActionChecked::CreateLambda([Settings]()
+		{
+			return Settings->bEnableCollision;
+		}));
+
+	MenuBuilder.AddMenuEntry(LOCTEXT("DefaultMeshObjectEnableCollision", "Enable Collision"), TAttribute<FText>(), FSlateIcon(), ToggleEnableCollisionAction, NAME_None, EUserInterfaceActionType::ToggleButton);
+
+	MenuBuilder.BeginSection("CollisionTypeSection", LOCTEXT("CollisionTypeSection", "Collision Mode"));
+
+	auto MakeCollisionTypeAction = [Settings](ECollisionTraceFlag FlagType)
+	{
+		return FUIAction(
+			FExecuteAction::CreateLambda([Settings, FlagType]
+			{
+				Settings->CollisionMode = FlagType;
+				Settings->SaveConfig();
+			}),
+			FCanExecuteAction::CreateLambda([Settings]() { return Settings->bEnableCollision; }),
+			FIsActionChecked::CreateLambda([Settings, FlagType]()
+			{
+				return Settings->CollisionMode == FlagType;
+			}));
+	};
+	const FUIAction CollisionMode_Default = MakeCollisionTypeAction(ECollisionTraceFlag::CTF_UseDefault);
+	MenuBuilder.AddMenuEntry(LOCTEXT("DefaultMeshObject_CollisionDefault", "Default"), LOCTEXT("DefaultMeshObject_CollisionDefault_Tooltip", "Configure the new Mesh Object to have the Project Default collision settings (CTF_UseDefault)"),
+		FSlateIcon(), CollisionMode_Default, NAME_None, EUserInterfaceActionType::ToggleButton);
+	const FUIAction CollisionMode_Both = MakeCollisionTypeAction(ECollisionTraceFlag::CTF_UseSimpleAndComplex);
+	MenuBuilder.AddMenuEntry(LOCTEXT("DefaultMeshObject_CollisionBoth", "Simple And Complex"), LOCTEXT("DefaultMeshObject_CollisionBoth_Tooltip", "Configure the new Mesh Object to have both Simple and Complex collision (CTF_UseSimpleAndComplex)"),
+		FSlateIcon(), CollisionMode_Both, NAME_None, EUserInterfaceActionType::ToggleButton);
+	const FUIAction CollisionMode_NoComplex = MakeCollisionTypeAction(ECollisionTraceFlag::CTF_UseSimpleAsComplex);
+	MenuBuilder.AddMenuEntry(LOCTEXT("DefaultMeshObject_CollisionNoComplex", "Simple Only"), LOCTEXT("DefaultMeshObject_CollisionNoComplex_Tooltip", "Configure the new Mesh Object to have only Simple collision (CTF_UseSimpleAsComplex)"),
+		FSlateIcon(), CollisionMode_NoComplex, NAME_None, EUserInterfaceActionType::ToggleButton);
+	const FUIAction CollisionMode_ComplexOnly = MakeCollisionTypeAction(ECollisionTraceFlag::CTF_UseComplexAsSimple);
+	MenuBuilder.AddMenuEntry(LOCTEXT("DefaultMeshObject_CollisionComplexOnly", "Complex Only"), LOCTEXT("DefaultMeshObject_CollisionComplexOnly_Tooltip", "Configure the new Mesh Object to have only Complex collision (CTF_UseComplexAsSimple)"),
+		FSlateIcon(), CollisionMode_ComplexOnly, NAME_None, EUserInterfaceActionType::ToggleButton);
+
+	MenuBuilder.EndSection();
+}
+
+
+void MakeSubMenu_DefaultMeshObjectType(FMenuBuilder& MenuBuilder)
+{
+	UModelingToolsEditorModeSettings* Settings = GetMutableDefault<UModelingToolsEditorModeSettings>();
+
+	MenuBuilder.BeginSection("DefaultObjecTypeSection", LOCTEXT("DefaultObjecTypeSection", "Default Object Type"));
+
+	MenuBuilder.AddMenuEntry(
+		LOCTEXT("DefaultMeshObjectType_StaticMesh", "Static Mesh"), 
+		TAttribute<FText>(), 
+		FSlateIcon(), 
+		FUIAction(
+			FExecuteAction::CreateLambda([Settings]
+			{
+				Settings->DefaultMeshObjectType = EModelingModeDefaultMeshObjectType::StaticMeshAsset;
+				UCreateMeshObjectTypeProperties::DefaultObjectTypeIdentifier = UCreateMeshObjectTypeProperties::StaticMeshIdentifier;
+				Settings->SaveConfig();
+			}),
+			FCanExecuteAction(),
+			FIsActionChecked::CreateLambda([Settings]()
+			{
+				return Settings->DefaultMeshObjectType == EModelingModeDefaultMeshObjectType::StaticMeshAsset;
+			})),
+		NAME_None, EUserInterfaceActionType::ToggleButton);
+
+	MenuBuilder.AddMenuEntry(
+		LOCTEXT("DefaultMeshObjectType_DynamicMesh", "Dynamic Mesh Actor"),
+		TAttribute<FText>(), 
+		FSlateIcon(), 
+		FUIAction(
+			FExecuteAction::CreateLambda([Settings]
+			{
+				Settings->DefaultMeshObjectType = EModelingModeDefaultMeshObjectType::DynamicMeshActor;
+				UCreateMeshObjectTypeProperties::DefaultObjectTypeIdentifier = UCreateMeshObjectTypeProperties::DynamicMeshActorIdentifier;
+				Settings->SaveConfig();
+			}),
+			FCanExecuteAction(),
+			FIsActionChecked::CreateLambda([Settings]()
+			{
+				return Settings->DefaultMeshObjectType == EModelingModeDefaultMeshObjectType::DynamicMeshActor;
+			})),
+		NAME_None, EUserInterfaceActionType::ToggleButton );
+
+	MenuBuilder.AddMenuEntry(
+		LOCTEXT("DefaultMeshObjectType_Volume", "Gameplay Volume"),
+		TAttribute<FText>(), 
+		FSlateIcon(), 
+		FUIAction(
+			FExecuteAction::CreateLambda([Settings]
+			{
+				Settings->DefaultMeshObjectType = EModelingModeDefaultMeshObjectType::VolumeActor;
+				UCreateMeshObjectTypeProperties::DefaultObjectTypeIdentifier = UCreateMeshObjectTypeProperties::VolumeIdentifier;
+				Settings->SaveConfig();
+			}),
+			FCanExecuteAction(),
+			FIsActionChecked::CreateLambda([Settings]()
+			{
+				return Settings->DefaultMeshObjectType == EModelingModeDefaultMeshObjectType::VolumeActor;
+			})),
+		NAME_None, EUserInterfaceActionType::ToggleButton);
+
+
+	MenuBuilder.EndSection();
+
+	MenuBuilder.AddMenuSeparator();
+
+	MakeSubMenu_DefaultMeshObjectPhysicsSettings(MenuBuilder);
+}
+
+
+
+TSharedRef<SWidget> MakeMenu_ModelingModeConfigSettings(FModelingToolsEditorModeToolkit* Toolkit)
+{
+	FMenuBuilder MenuBuilder(true, TSharedPtr<FUICommandList>());
+
+	MenuBuilder.BeginSection("Section_Selection", LOCTEXT("Section_Selection", "Selection"));
+	MakeSubMenu_ModeToggles(MenuBuilder);
+	MenuBuilder.EndSection();
+
+	MenuBuilder.BeginSection("Section_Gizmo", LOCTEXT("Section_Gizmo", "Gizmo"));
+	MakeSubMenu_GizmoVisibilityMode(Toolkit, MenuBuilder);
+	MenuBuilder.EndSection();
+
+	MenuBuilder.BeginSection("Section_MeshObjects", LOCTEXT("Section_MeshObjects", "New Mesh Objects"));
+	MenuBuilder.AddSubMenu(
+		LOCTEXT("MeshObjectTypeSubMenu", "New Mesh Settings"), LOCTEXT("MeshObjectTypeSubMenu_ToolTip", "Configure default settings for new Mesh Object Types"),
+		FNewMenuDelegate::CreateLambda([=](FMenuBuilder& SubMenuBuilder) {
+			MakeSubMenu_DefaultMeshObjectType(SubMenuBuilder);
+		}));
+	MenuBuilder.EndSection();
+
+	MenuBuilder.BeginSection("Section_Settings", LOCTEXT("Section_Settings", "Quick Settings"));
+	MenuBuilder.AddSubMenu(
+		LOCTEXT("QuickSettingsSubMenu", "Jump To Settings"), LOCTEXT("QuickSettingsSubMenu_ToolTip", "Jump to sections of the Settings dialogs relevant to Modeling Mode"),
+		FNewMenuDelegate::CreateLambda([=](FMenuBuilder& SubMenuBuilder) {
+			MakeSubMenu_QuickSettings(SubMenuBuilder);
+			}));
+	MenuBuilder.EndSection();
+
+	TSharedRef<SWidget> MenuWidget = MenuBuilder.MakeWidget();
+	return MenuWidget;
+}
+
+
+
+
+
+void MakeSubMenu_Selection_DragMode(FModelingToolsEditorModeToolkit* Toolkit, FMenuBuilder& MenuBuilder)
+{
+	UModelingSelectionInteraction* SelectionInteraction = Cast<UModelingToolsEditorMode>(Toolkit->GetScriptableEditorMode())->GetSelectionInteraction();
+
+	auto MakeDragModeOptionAction = [SelectionInteraction](EModelingSelectionInteraction_DragMode DragMode)
+	{
+		return FUIAction(
+			FExecuteAction::CreateLambda([SelectionInteraction, DragMode]
+			{
+				SelectionInteraction->SetActiveDragMode(DragMode);
+			}),
+			FCanExecuteAction(),
+			FIsActionChecked::CreateLambda([SelectionInteraction, DragMode]()
+			{
+				return SelectionInteraction->GetActiveDragMode() == DragMode;
+			}));
+	};
+
+	MenuBuilder.AddMenuEntry(
+		LOCTEXT("Selection_DragInput_None", "None"), 
+		LOCTEXT("Selection_DragInput_None_Tooltip", "No drag input"),
+		FSlateIcon(), MakeDragModeOptionAction(EModelingSelectionInteraction_DragMode::NoDragInteraction), NAME_None, EUserInterfaceActionType::ToggleButton);
+
+	MenuBuilder.AddMenuEntry(
+		LOCTEXT("Selection_DragInput_Path", "Path"), 
+		LOCTEXT("Selection_DragInput_Path_Tooltip", "Path Drag Input"),
+		FSlateIcon(), MakeDragModeOptionAction(EModelingSelectionInteraction_DragMode::PathInteraction), NAME_None, EUserInterfaceActionType::ToggleButton);
+
+	MenuBuilder.AddMenuEntry(
+		LOCTEXT("Selection_DragInput_Marquee", "Rectangle"), 
+		LOCTEXT("Selection_DragInput_Marquee_Tooltip", "Rectangle Marquee"),
+		FSlateIcon(), MakeDragModeOptionAction(EModelingSelectionInteraction_DragMode::RectangleMarqueeInteraction), NAME_None, EUserInterfaceActionType::ToggleButton);
+
+}
+
+
+void MakeSubMenu_Selection_MeshType(FModelingToolsEditorModeToolkit* Toolkit, FMenuBuilder& MenuBuilder)
+{
+	UModelingToolsEditorMode* ModelingMode = Cast<UModelingToolsEditorMode>(Toolkit->GetScriptableEditorMode());
+	FUIAction ToggleVolumesAction(
+		FExecuteAction::CreateLambda([ModelingMode]
+		{
+			ModelingMode->bEnableVolumeElementSelection = ! ModelingMode->bEnableVolumeElementSelection;
+		}),
+		FCanExecuteAction(),
+		FIsActionChecked::CreateLambda([ModelingMode]()
+		{
+			return ModelingMode->bEnableVolumeElementSelection;
+		}));
+	FUIAction ToggleStaticMeshesAction(
+		FExecuteAction::CreateLambda([ModelingMode]
+		{
+			ModelingMode->bEnableStaticMeshElementSelection = ! ModelingMode->bEnableStaticMeshElementSelection;
+		}),
+		FCanExecuteAction(),
+		FIsActionChecked::CreateLambda([ModelingMode]()
+		{
+			return ModelingMode->bEnableStaticMeshElementSelection;
+		}));
+
+
+	MenuBuilder.AddMenuEntry(
+		LOCTEXT("Selection_MeshTypes_Volumes", "Volumes"), 
+		LOCTEXT("Selection_MeshTypes_Volumes_Tooltip", "Toggle whether Volume mesh elements can be selected in the Viewport"),
+		FSlateIcon(), ToggleVolumesAction, NAME_None, EUserInterfaceActionType::ToggleButton);
+
+	MenuBuilder.AddMenuEntry(
+		LOCTEXT("Selection_MeshTypes_StaticMesh", "Static Meshes"), 
+		LOCTEXT("Selection_MeshTypes_StaticMesh_Tooltip", "Toggle whether Static Mesh mesh elements can be selected in the Viewport"),
+		FSlateIcon(), ToggleStaticMeshesAction, NAME_None, EUserInterfaceActionType::ToggleButton);
+
+}
+
+
+
+TSharedRef<SWidget> MakeMenu_SelectionConfigSettings(FModelingToolsEditorModeToolkit* Toolkit)
+{
+	FMenuBuilder MenuBuilder(true, TSharedPtr<FUICommandList>());
+
+	MenuBuilder.BeginSection("Section_DragMode", LOCTEXT("Section_DragMode", "Drag Mode"));
+	MakeSubMenu_Selection_DragMode(Toolkit, MenuBuilder);
+	MenuBuilder.EndSection();
+
+	MenuBuilder.BeginSection("Section_MeshTypes", LOCTEXT("Section_MeshTypes", "Selectable Mesh Types"));
+	MakeSubMenu_Selection_MeshType(Toolkit, MenuBuilder);
+	MenuBuilder.EndSection();
+
+	TSharedRef<SWidget> MenuWidget = MenuBuilder.MakeWidget();
+	return MenuWidget;
+}
+
+} // end namespace UELocal
+
+
+void FModelingToolsEditorModeToolkit::MakeSelectionPaletteOverlayWidget()
+{
+	FVerticalToolBarBuilder ToolbarBuilder( 
+		GetToolkitCommands(), 
+		FMultiBoxCustomization::None, 
+		TSharedPtr<FExtender>(), /*InForceSmallIcons=*/ true);
+	ToolbarBuilder.SetLabelVisibility(EVisibility::Collapsed);
+	ToolbarBuilder.SetStyle(&FAppStyle::Get(), "EditorViewportToolBar");
+
+	const FModelingToolsManagerCommands& Commands = FModelingToolsManagerCommands::Get();
+
+	ToolbarBuilder.AddToolBarButton(Commands.MeshSelectionModeAction_NoSelection);
+	ToolbarBuilder.AddWidget( SNew(SSpacer).Size(FVector2D(1,4)) );
+	ToolbarBuilder.AddToolBarButton(Commands.MeshSelectionModeAction_MeshTriangles);
+	ToolbarBuilder.AddToolBarButton(Commands.MeshSelectionModeAction_MeshEdges);
+	ToolbarBuilder.AddToolBarButton(Commands.MeshSelectionModeAction_MeshVertices);
+	ToolbarBuilder.AddWidget( SNew(SSpacer).Size(FVector2D(1,4)) );
+	ToolbarBuilder.AddToolBarButton(Commands.MeshSelectionModeAction_GroupFaces);
+	ToolbarBuilder.AddToolBarButton(Commands.MeshSelectionModeAction_GroupEdges);
+	ToolbarBuilder.AddToolBarButton(Commands.MeshSelectionModeAction_GroupCorners);
+
+	ToolbarBuilder.AddWidget(
+		SNew(SComboButton)
+		.HasDownArrow(false)
+		.MenuPlacement(EMenuPlacement::MenuPlacement_MenuRight)
+		.ComboButtonStyle(FAppStyle::Get(), "SimpleComboButton")
+		.OnGetMenuContent(FOnGetContent::CreateLambda([this]()
+		{
+			return UELocal::MakeMenu_SelectionConfigSettings(this);
+		}))
+		.ContentPadding(FMargin(1.0f, 1.0f))
+		.ButtonContent()
+		[
+			SNew(SImage)
+			.Image(FModelingToolsEditorModeStyle::Get()->GetBrush("ModelingModeSelection.More_Right"))
+		] );
+
+	ToolbarBuilder.AddWidget( SNew(SSpacer).Size(FVector2D(1,10)) );
+
+	ToolbarBuilder.AddWidget(
+		SNew(SComboButton)
+		.HasDownArrow(false)
+		.MenuPlacement(EMenuPlacement::MenuPlacement_MenuRight)
+		.ComboButtonStyle(FAppStyle::Get(), "SimpleComboButton")
+		.OnGetMenuContent(FOnGetContent::CreateLambda([this]()
+		{
+			return UELocal::MakeMenu_ModelingModeConfigSettings(this);
+		}))
+		.ContentPadding(FMargin(3.0f, 1.0f))
+		.ButtonContent()
+		[
+			SNew(SImage)
+			.Image(FModelingToolsEditorModeStyle::Get()->GetBrush("ModelingMode.DefaultSettings"))
+		] );
+
+
+	TSharedRef<SWidget> Toolbar = ToolbarBuilder.MakeWidget();
+
+
+	SAssignNew(SelectionPaletteOverlayWidget, SVerticalBox)
+	+SVerticalBox::Slot()
+	.HAlign(HAlign_Left)
+	.VAlign(VAlign_Center)
+	.Padding(FMargin(2.0f, 0.0f, 0.f, 0.f))
+	[
+		SNew(SBorder)
+		.BorderImage(FAppStyle::Get().GetBrush("EditorViewport.OverlayBrush"))
+		.Padding(FMargin(3.0f, 6.0f, 3.f, 6.f))
+		[
+			SNew(SVerticalBox)
+
+			+ SVerticalBox::Slot()
+			.AutoHeight()
+			.Padding(FMargin(0.0, 0.f, 0.f, 0.f))
+			[
+				Toolbar
+			]
+		]
+	];
+
+}
+
+
+#undef LOCTEXT_NAMESPACE
\ No newline at end of file
diff --git a/Engine/Plugins/Editor/ModelingToolsEditorMode/Source/ModelingToolsEditorMode/Private/ModelingSelectionInteraction.cpp b/Engine/Plugins/Editor/ModelingToolsEditorMode/Source/ModelingToolsEditorMode/Private/ModelingSelectionInteraction.cpp
--- a/Engine/Plugins/Editor/ModelingToolsEditorMode/Source/ModelingToolsEditorMode/Private/ModelingSelectionInteraction.cpp
+++ b/Engine/Plugins/Editor/ModelingToolsEditorMode/Source/ModelingToolsEditorMode/Private/ModelingSelectionInteraction.cpp
@@ -9,8 +9,11 @@
 #include "SceneQueries/SceneSnappingManager.h"
 
 #include "BaseBehaviors/SingleClickBehavior.h"
+#include "BaseBehaviors/SingleClickOrDragBehavior.h"
 #include "BaseBehaviors/MouseHoverBehavior.h"
 
+#include "Mechanics/RectangleMarqueeMechanic.h"
+
 
 using namespace UE::Geometry;
 
@@ -24,19 +27,31 @@ void UModelingSelectionInteraction::Initialize(
 	CanChangeSelectionCallback = MoveTemp(CanChangeSelectionCallbackIn);
 	ExternalHitCaptureCallback = MoveTemp(ExternalHitCaptureCallbackIn);
 
+	// set up rectangle marquee interaction
+	RectangleMarqueeInteraction = NewObject<URectangleMarqueeInteraction>(this);
+	RectangleMarqueeInteraction->OnDragRectangleFinished.AddUObject(this, &UModelingSelectionInteraction::OnMarqueeRectangleFinished);
+
+	// set up path selection interaction
+	PathSelectionInteraction = NewObject<UPathSelectionInteraction>(this);
+	PathSelectionInteraction->Setup(this);
+
 	// create click behavior and set ourselves as click target
-	ClickBehavior = NewObject<USingleClickInputBehavior>();
-	ClickBehavior->Modifiers.RegisterModifier(AddToSelectionModifier, FInputDeviceState::IsShiftKeyDown);
-	ClickBehavior->Modifiers.RegisterModifier(ToggleSelectionModifier, FInputDeviceState::IsCtrlKeyDown);
-	ClickBehavior->Initialize(this);
+	ClickOrDragBehavior = NewObject<USingleClickOrDragInputBehavior>();
+	ClickOrDragBehavior->Modifiers.RegisterModifier(AddToSelectionModifier, FInputDeviceState::IsShiftKeyDown);
+	ClickOrDragBehavior->Modifiers.RegisterModifier(ToggleSelectionModifier, FInputDeviceState::IsCtrlKeyDown);
+	ClickOrDragBehavior->bBeginDragIfClickTargetNotHit = false;
+	ClickOrDragBehavior->Initialize(this, PathSelectionInteraction);
 
 	HoverBehavior = NewObject<UMouseHoverBehavior>(this);
 	HoverBehavior->Initialize(this);
 
 	BehaviorSet = NewObject<UInputBehaviorSet>();
-	BehaviorSet->Add(ClickBehavior, this);
+	BehaviorSet->Add(ClickOrDragBehavior, this);
 	BehaviorSet->Add(HoverBehavior, this);
 
+	// configure drag mode of ClickOrDragBehavior
+	UpdateActiveDragMode();
+
 	TransformProxy = NewObject<UTransformProxy>(this);
 	// todo: make this repositionable etc. Maybe make this function a delegate? or allow caller to provide the gizmo?
 	TransformGizmo = UE::TransformGizmoUtil::Create3AxisTransformGizmo(
@@ -68,6 +83,38 @@ void UModelingSelectionInteraction::Shutdown()
 }
 
 
+void UModelingSelectionInteraction::DrawHUD(FCanvas* Canvas, IToolsContextRenderAPI* RenderAPI)
+{
+	if (RectangleMarqueeInteraction)
+	{
+		RectangleMarqueeInteraction->DrawHUD(Canvas, RenderAPI);
+	}
+}
+
+void UModelingSelectionInteraction::SetActiveDragMode(EModelingSelectionInteraction_DragMode NewMode)
+{
+	if (ActiveDragMode != NewMode)
+	{
+		ActiveDragMode = NewMode;
+		UpdateActiveDragMode();
+	}
+}
+
+void UModelingSelectionInteraction::UpdateActiveDragMode()
+{
+	if (ActiveDragMode == EModelingSelectionInteraction_DragMode::PathInteraction)
+	{
+		ClickOrDragBehavior->SetDragTarget(PathSelectionInteraction);
+	}
+	else if (ActiveDragMode == EModelingSelectionInteraction_DragMode::RectangleMarqueeInteraction)
+	{
+		ClickOrDragBehavior->SetDragTarget(RectangleMarqueeInteraction);
+	}
+	else
+	{
+		ClickOrDragBehavior->SetDragTarget(nullptr);
+	}
+}
 
 
 void UModelingSelectionInteraction::OnUpdateModifierState(int ModifierID, bool bIsOn)
@@ -196,6 +243,18 @@ void UModelingSelectionInteraction::OnClicked(const FInputDeviceRay& ClickPos)
 		return;
 	}
 
+	FGeometrySelectionUpdateConfig UpdateConfig = GetActiveSelectionUpdateConfig();
+
+	FGeometrySelectionUpdateResult Result;
+	SelectionManager->UpdateSelectionViaRaycast(
+		ClickPos.WorldRay,
+		UpdateConfig,
+		Result);
+}
+
+
+FGeometrySelectionUpdateConfig UModelingSelectionInteraction::GetActiveSelectionUpdateConfig() const
+{
 	FGeometrySelectionUpdateConfig UpdateConfig;
 	UpdateConfig.ChangeType = EGeometrySelectionChangeType::Replace;
 	if (bAddToSelectionEnabled)
@@ -206,18 +265,11 @@ void UModelingSelectionInteraction::OnClicked(const FInputDeviceRay& ClickPos)
 	{
 		UpdateConfig.ChangeType = EGeometrySelectionChangeType::Remove;
 	}
-
-	FGeometrySelectionUpdateResult Result;
-	SelectionManager->UpdateSelectionViaRaycast(
-		ClickPos.WorldRay,
-		UpdateConfig,
-		Result);
+	return UpdateConfig;
 }
 
 
 
-
-
 FInputRayHit UModelingSelectionInteraction::BeginHoverSequenceHitTest(const FInputDeviceRay& PressPos)
 {
 	if ((SelectionManager->GetMeshTopologyMode() == UGeometrySelectionManager::EMeshTopologyMode::None)
@@ -365,8 +417,78 @@ void UModelingSelectionInteraction::ApplyPendingTransformInteractions()
 		}
 	}
 
+	bGizmoUpdatePending = false;
+}
 
 
 
-	bGizmoUpdatePending = false;
+void UModelingSelectionInteraction::OnMarqueeRectangleFinished(const FCameraRectangle& Rectangle, bool bCancelled)
+{
+
+}
+
+
+
+
+void UPathSelectionInteraction::Setup(UModelingSelectionInteraction* SelectionInteractionIn)
+{
+
+	SelectionInteraction = SelectionInteractionIn;
+}
+
+
+FInputRayHit UPathSelectionInteraction::CanBeginClickDragSequence(const FInputDeviceRay& PressPos)
+{
+	if (SelectionInteraction->GetSelectionManager()->CanBeginTrackedSelectionChange())
+	{
+		return SelectionInteraction->IsHitByClick(PressPos);
+	}
+	return FInputRayHit();
+}
+
+void UPathSelectionInteraction::OnClickPress(const FInputDeviceRay& PressPos)
+{
+	FGeometrySelectionUpdateConfig UpdateConfig = SelectionInteraction->GetActiveSelectionUpdateConfig();
+	UGeometrySelectionManager* SelectionManager = SelectionInteraction->GetSelectionManager();
+
+	bool bInitialClear = false;
+	if (UpdateConfig.ChangeType == EGeometrySelectionChangeType::Replace)
+	{
+		bInitialClear = true;
+		UpdateConfig.ChangeType = EGeometrySelectionChangeType::Add;
+	}
+
+	if (SelectionManager->BeginTrackedSelectionChange(UpdateConfig, bInitialClear))
+	{
+		FGeometrySelectionUpdateResult Result;
+		SelectionManager->AccumulateSelectionUpdate_Raycast(PressPos.WorldRay, Result);
+	}
+}
+
+void UPathSelectionInteraction::OnClickDrag(const FInputDeviceRay& DragPos)
+{
+	UGeometrySelectionManager* SelectionManager = SelectionInteraction->GetSelectionManager();
+	if (SelectionManager->IsInTrackedSelectionChange())
+	{
+		FGeometrySelectionUpdateResult Result;
+		SelectionManager->AccumulateSelectionUpdate_Raycast(DragPos.WorldRay, Result);		
+	}
+}
+
+void UPathSelectionInteraction::OnClickRelease(const FInputDeviceRay& ReleasePos)
+{
+	UGeometrySelectionManager* SelectionManager = SelectionInteraction->GetSelectionManager();
+	if (SelectionManager->IsInTrackedSelectionChange())
+	{
+		SelectionManager->EndTrackedSelectionChange();
+	}
+}
+
+void UPathSelectionInteraction::OnTerminateDragSequence()
+{
+	UGeometrySelectionManager* SelectionManager = SelectionInteraction->GetSelectionManager();
+	if (SelectionManager->IsInTrackedSelectionChange())
+	{
+		SelectionManager->EndTrackedSelectionChange();
+	}
 }
\ No newline at end of file
diff --git a/Engine/Plugins/Editor/ModelingToolsEditorMode/Source/ModelingToolsEditorMode/Private/ModelingToolsEditorMode.cpp b/Engine/Plugins/Editor/ModelingToolsEditorMode/Source/ModelingToolsEditorMode/Private/ModelingToolsEditorMode.cpp
--- a/Engine/Plugins/Editor/ModelingToolsEditorMode/Source/ModelingToolsEditorMode/Private/ModelingToolsEditorMode.cpp
+++ b/Engine/Plugins/Editor/ModelingToolsEditorMode/Source/ModelingToolsEditorMode/Private/ModelingToolsEditorMode.cpp
@@ -94,6 +94,8 @@
 #include "SplitMeshesTool.h"
 #include "PatternTool.h"
 
+//#include "Polymodeling/ExtrudeMeshSelectionTool.h"
+
 #include "Physics/PhysicsInspectorTool.h"
 #include "Physics/SetCollisionGeometryTool.h"
 #include "Physics/ExtractCollisionGeometryTool.h"
@@ -126,8 +128,13 @@
 #include "Selections/GeometrySelection.h"
 #include "Selection/GeometrySelectionManager.h"
 #include "Selection/DynamicMeshSelector.h"
+#include "Selection/VolumeSelector.h"
+#include "Selection/StaticMeshSelector.h"
 #include "ModelingSelectionInteraction.h"
 #include "DynamicMeshActor.h"
+#include "GameFramework/Volume.h"
+#include "Components/BrushComponent.h"
+#include "Engine/StaticMeshActor.h"
 
 #include UE_INLINE_GENERATED_CPP_BY_NAME(ModelingToolsEditorMode)
 
@@ -382,8 +389,11 @@ void UModelingToolsEditorMode::Enter()
 		SelectionManager = NewObject<UGeometrySelectionManager>(GetToolManager());
 		SelectionManager->Initialize(GetInteractiveToolsContext(), GetToolManager()->GetContextTransactionsAPI());
 		SelectionManager->RegisterSelectorFactory(MakeUnique<FDynamicMeshComponentSelectorFactory>());
+		SelectionManager->RegisterSelectorFactory(MakeUnique<FBrushComponentSelectorFactory>());
+		SelectionManager->RegisterSelectorFactory(MakeUnique<FStaticMeshComponentSelectorFactory>());
 
 		GetInteractiveToolsContext()->OnRender.AddUObject(this, &UModelingToolsEditorMode::OnToolsContextRender);
+		GetInteractiveToolsContext()->OnDrawHUD.AddUObject(this, &UModelingToolsEditorMode::OnToolsContextDrawHUD);
 		// this is hopefully temporary? kinda gross...
 		GetInteractiveToolsContext()->ContextObjectStore->AddContextObject(SelectionManager);
 
@@ -726,6 +736,9 @@ void UModelingToolsEditorMode::Enter()
 	RegisterPolyModelActionTool(EEditMeshPolygonsToolActions::CutFaces, ToolManagerCommands.BeginPolyModelTool_CutFaces, TEXT("PolyEdit_CutFaces"));
 
 
+	//RegisterTool(ToolManagerCommands.BeginSelectionAction_Extrude, TEXT("BeginSelectionExtrudeTool"), NewObject<UExtrudeMeshSelectionToolBuilder>());
+
+
 	// set up selection type toggles
 	auto RegisterSelectionTopologyType = [this](UGeometrySelectionManager::EMeshTopologyMode TopoMode, TSharedPtr<FUICommandInfo> UICommand)
 	{
@@ -758,6 +771,31 @@ void UModelingToolsEditorMode::Enter()
 	}
 
 
+	auto RegisterSelectionMode = [this](UGeometrySelectionManager::EMeshTopologyMode TopoMode, UE::Geometry::EGeometryElementType ElementMode, TSharedPtr<FUICommandInfo> UICommand)
+	{
+		Toolkit->GetToolkitCommands()->MapAction(UICommand,
+			FExecuteAction::CreateLambda([this, TopoMode, ElementMode]() {
+				GetToolManager()->GetContextTransactionsAPI()->BeginUndoTransaction(LOCTEXT("ChangeSelectionMode", "Selection Mode"));
+				GetSelectionManager()->SetMeshTopologyMode(TopoMode); 
+				GetSelectionManager()->SetSelectionElementType(ElementMode);
+				GetToolManager()->GetContextTransactionsAPI()->EndUndoTransaction();
+			}),
+			FCanExecuteAction(),
+			FIsActionChecked::CreateLambda([this, TopoMode, ElementMode]() { return GetSelectionManager()->GetMeshTopologyMode() == TopoMode && GetSelectionManager()->GetSelectionElementType() == ElementMode; }),
+			EUIActionRepeatMode::RepeatDisabled);
+	};
+	if (GetSelectionManager() != nullptr)
+	{
+		RegisterSelectionMode(UGeometrySelectionManager::EMeshTopologyMode::None, UGeometrySelectionManager::EGeometryElementType::Face, ToolManagerCommands.MeshSelectionModeAction_NoSelection);
+		RegisterSelectionMode(UGeometrySelectionManager::EMeshTopologyMode::Triangle, UGeometrySelectionManager::EGeometryElementType::Face, ToolManagerCommands.MeshSelectionModeAction_MeshTriangles);
+		RegisterSelectionMode(UGeometrySelectionManager::EMeshTopologyMode::Triangle, UGeometrySelectionManager::EGeometryElementType::Vertex, ToolManagerCommands.MeshSelectionModeAction_MeshVertices);
+		RegisterSelectionMode(UGeometrySelectionManager::EMeshTopologyMode::Triangle, UGeometrySelectionManager::EGeometryElementType::Edge, ToolManagerCommands.MeshSelectionModeAction_MeshEdges);
+		RegisterSelectionMode(UGeometrySelectionManager::EMeshTopologyMode::Polygroup, UGeometrySelectionManager::EGeometryElementType::Face, ToolManagerCommands.MeshSelectionModeAction_GroupFaces);
+		RegisterSelectionMode(UGeometrySelectionManager::EMeshTopologyMode::Polygroup, UGeometrySelectionManager::EGeometryElementType::Vertex, ToolManagerCommands.MeshSelectionModeAction_GroupCorners);
+		RegisterSelectionMode(UGeometrySelectionManager::EMeshTopologyMode::Polygroup, UGeometrySelectionManager::EGeometryElementType::Edge, ToolManagerCommands.MeshSelectionModeAction_GroupEdges);
+	}
+
+
 	// this function registers and tracks an active UGeometrySelectionEditCommand and it's associated UICommand
 	auto RegisterSelectionCommand = [&](UGeometrySelectionEditCommand* Command, TSharedPtr<FUICommandInfo> UICommand)
 	{
@@ -767,6 +805,7 @@ void UModelingToolsEditorMode::Enter()
 			FExecuteAction::CreateUObject(GetSelectionManager(), &UGeometrySelectionManager::ExecuteSelectionCommand, Command),
 			FCanExecuteAction::CreateUObject(GetSelectionManager(), &UGeometrySelectionManager::CanExecuteSelectionCommand, Command),
 			FIsActionChecked(),
+			FIsActionButtonVisible::CreateUObject(GetSelectionManager(), &UGeometrySelectionManager::CanExecuteSelectionCommand, Command),
 			EUIActionRepeatMode::RepeatDisabled);
 	};
 
@@ -873,6 +912,9 @@ void UModelingToolsEditorMode::Enter()
 	SelectionModifiedEventHandle = GetModeManager()->GetSelectedActors()->SelectionChangedEvent.AddLambda(
 		[this](const UObject* Object)  { UpdateSelectionManagerOnEditorSelectionChange(); } );
 
+	// initialize SelectionManager w/ active selection
+	UpdateSelectionManagerOnEditorSelectionChange(true);
+
 }
 
 void UModelingToolsEditorMode::RegisterUVEditor()
@@ -1044,6 +1086,14 @@ void UModelingToolsEditorMode::OnToolsContextRender(IToolsContextRenderAPI* Rend
 	}
 }
 
+void UModelingToolsEditorMode::OnToolsContextDrawHUD(FCanvas* Canvas, IToolsContextRenderAPI* RenderAPI)
+{
+	if (SelectionInteraction)
+	{
+		SelectionInteraction->DrawHUD(Canvas, RenderAPI);
+	}
+}
+
 bool UModelingToolsEditorMode::ShouldToolStartBeAllowed(const FString& ToolIdentifier) const
 {
 	if (UInteractiveToolManager* Manager = GetToolManager())
@@ -1089,7 +1139,7 @@ bool UModelingToolsEditorMode::TestForEditorGizmoHit(const FInputDeviceRay& Clic
 
 
 
-void UModelingToolsEditorMode::UpdateSelectionManagerOnEditorSelectionChange()
+void UModelingToolsEditorMode::UpdateSelectionManagerOnEditorSelectionChange(bool bEnteringMode)
 {
 	if (!SelectionManager) return;
 
@@ -1100,31 +1150,60 @@ void UModelingToolsEditorMode::UpdateSelectionManagerOnEditorSelectionChange()
 		return;
 	}
 
+	// figure out what is selected and then extract valid selection identifiers if there are any
+	TArray<AActor*> SelectedActors;
+	GEditor->GetSelectedActors()->GetSelectedObjects<AActor>(SelectedActors); 
+
+	// this should probably not be the place we make this kind of decision...but for now it will do
+	TArray<FGeometryIdentifier> ValidIdentifiers;
+	for ( AActor* Actor : SelectedActors)
+	{
+		if (Cast<ADynamicMeshActor>(Actor) != nullptr)
+		{
+			ValidIdentifiers.Add( FGeometryIdentifier::PrimitiveComponent(
+				Cast<ADynamicMeshActor>(Actor)->GetDynamicMeshComponent(), FGeometryIdentifier::EObjectType::DynamicMeshComponent) );
+		}
+		else if (bEnableVolumeElementSelection && Cast<AVolume>(Actor) != nullptr)
+		{
+			ValidIdentifiers.Add( FGeometryIdentifier::PrimitiveComponent(
+				Cast<AVolume>(Actor)->GetBrushComponent(), FGeometryIdentifier::EObjectType::BrushComponent) );
+		}
+		else if (bEnableStaticMeshElementSelection && Cast<AStaticMeshActor>(Actor) != nullptr)
+		{
+			// StaticMeshComponent can exist under any Actor, should not rely on Actor selection...
+			ValidIdentifiers.Add( FGeometryIdentifier::PrimitiveComponent(
+				Cast<AStaticMeshActor>(Actor)->GetStaticMeshComponent(), FGeometryIdentifier::EObjectType::StaticMeshComponent) );
+		}
+	}
+
+	// This is gross. If we are entering the Mode, we need to update the SelectionManager w/ the current state. However this
+	// update needs to be undoable. Since we are not part of whatever Transaction was involved in changing modes, we are going
+	// to have to emit our own Transaction, which will then be an explicit undo step the user has to go through :(
+	bool bPendingCloseTransaction = false;
+	if (bEnteringMode && ValidIdentifiers.Num() > 0)
+	{
+		GetToolManager()->GetContextTransactionsAPI()->BeginUndoTransaction(LOCTEXT("InitializeSelection", "Initialize Selection"));
+		bPendingCloseTransaction = true;
+	}
+
 	// If Editor is creating a transaction, we assume we must be in a selection change.
 	// Need to handle all SelectionManager changes (deselect + change-targets) during
 	// the transaction so that undo/redo works properly.
+	// (note that if we are bEnteringMode, we just opened a transaction and so this branch will still be taken...)
 	bool bCreatingTransaction = (GUndo != nullptr);
 	if (bCreatingTransaction)
 	{
-		TArray<AActor*> SelectedActors;
-		GEditor->GetSelectedActors()->GetSelectedObjects<AActor>(SelectedActors); 
-
-		TArray<FGeometryIdentifier> ValidIdentifiers;
-		for ( AActor* Actor : SelectedActors)
-		{
-			if (Cast<ADynamicMeshActor>(Actor) != nullptr)
-			{
-				ValidIdentifiers.Add( FGeometryIdentifier::PrimitiveComponent(
-					Cast<ADynamicMeshActor>(Actor)->GetDynamicMeshComponent(), FGeometryIdentifier::EObjectType::DynamicMeshComponent) );
-			}
-		}
-
 		SelectionManager->SynchronizeActiveTargets(ValidIdentifiers,
 			[this]() {
 				SelectionManager->ClearSelection();
 			});
 	}
 
+	// close out transaction if it was still open
+	if (bPendingCloseTransaction)
+	{
+		GetToolManager()->GetContextTransactionsAPI()->EndUndoTransaction();
+	}
 }
 
 
diff --git a/Engine/Plugins/Editor/ModelingToolsEditorMode/Source/ModelingToolsEditorMode/Private/ModelingToolsEditorModeStyle.cpp b/Engine/Plugins/Editor/ModelingToolsEditorMode/Source/ModelingToolsEditorMode/Private/ModelingToolsEditorModeStyle.cpp
--- a/Engine/Plugins/Editor/ModelingToolsEditorMode/Source/ModelingToolsEditorMode/Private/ModelingToolsEditorModeStyle.cpp
+++ b/Engine/Plugins/Editor/ModelingToolsEditorMode/Source/ModelingToolsEditorMode/Private/ModelingToolsEditorModeStyle.cpp
@@ -93,6 +93,8 @@ void FModelingToolsEditorModeStyle::Initialize()
 			StyleSet->Set("LevelEditor.ModelingToolsMode.Small", new IMAGE_PLUGIN_BRUSH("Icons/icon_ModelingToolsEditorMode", FVector2D(20.0f, 20.0f)));
 		}
 
+		StyleSet->Set("ModelingMode.DefaultSettings", new FSlateImageBrush(StyleSet->RootToCoreContentDir(TEXT("../Editor/Slate/Icons/GeneralTools/Settings_40x.png")), Icon20x20));
+
 		// NOTE:  Old-style, need to be replaced: 
 		StyleSet->Set("ModelingToolsManagerCommands.CancelActiveTool", new IMAGE_PLUGIN_BRUSH("Icons/icon_ActiveTool_Cancel_40x", Icon20x20));
 		StyleSet->Set("ModelingToolsManagerCommands.CancelActiveTool.Small", new IMAGE_PLUGIN_BRUSH("Icons/icon_ActiveTool_Cancel_40x", Icon20x20));
@@ -318,6 +320,20 @@ void FModelingToolsEditorModeStyle::Initialize()
 		StyleSet->Set("ModelingToolsManagerCommands.BeginSkinWeightsPaintTool", new IMAGE_BRUSH_SVG("Icons/SkinWeightsPaint", Icon20x20));
 		StyleSet->Set("ModelingToolsManagerCommands.BeginSkinWeightsPaintTool.Small", new IMAGE_BRUSH_SVG("Icons/SkinWeightsPaint", Icon20x20));
 
+		StyleSet->Set("ModelingToolsManagerCommands.MeshSelectionModeAction_NoSelection", new FSlateImageBrush(StyleSet->RootToCoreContentDir(TEXT("../Editor/Slate/Icons/GeneralTools/Select_40x.png")), Icon20x20));
+		StyleSet->Set("ModelingToolsManagerCommands.MeshSelectionModeAction_MeshTriangles", new IMAGE_BRUSH_SVG("Icons/SelectionToolbar_Triangles", Icon20x20));
+		StyleSet->Set("ModelingToolsManagerCommands.MeshSelectionModeAction_MeshVertices", new IMAGE_BRUSH_SVG("Icons/SelectionToolbar_Vertices", Icon20x20));
+		StyleSet->Set("ModelingToolsManagerCommands.MeshSelectionModeAction_MeshEdges", new IMAGE_BRUSH_SVG("Icons/SelectionToolbar_Edges", Icon20x20));
+		StyleSet->Set("ModelingToolsManagerCommands.MeshSelectionModeAction_GroupFaces", new IMAGE_BRUSH_SVG("Icons/SelectionTriangles3", Icon20x20));
+		StyleSet->Set("ModelingToolsManagerCommands.MeshSelectionModeAction_GroupCorners", new IMAGE_BRUSH_SVG("Icons/SelectionVertices", Icon20x20));
+		StyleSet->Set("ModelingToolsManagerCommands.MeshSelectionModeAction_GroupEdges", new IMAGE_BRUSH_SVG("Icons/SelectionBorderEdges", Icon20x20));
+
+		StyleSet->Set("ModelingToolsManagerCommands.BeginSelectionAction_Delete", new FSlateImageBrush(StyleSet->RootToCoreContentDir(TEXT("../Editor/Slate/Icons/GeneralTools/Delete_40x.png")), Icon20x20));
+
+		StyleSet->Set("ModelingModeSelection.More_Right",  new IMAGE_BRUSH_SVG("Icons/SelectionToolbar_More", Icon20x20));
+
+
+
 
 		//
 		// Icons for brush falloffs in sculpt/etc tools
diff --git a/Engine/Plugins/Editor/ModelingToolsEditorMode/Source/ModelingToolsEditorMode/Private/ModelingToolsEditorModeToolkit.cpp b/Engine/Plugins/Editor/ModelingToolsEditorMode/Source/ModelingToolsEditorMode/Private/ModelingToolsEditorModeToolkit.cpp
--- a/Engine/Plugins/Editor/ModelingToolsEditorMode/Source/ModelingToolsEditorMode/Private/ModelingToolsEditorModeToolkit.cpp
+++ b/Engine/Plugins/Editor/ModelingToolsEditorMode/Source/ModelingToolsEditorMode/Private/ModelingToolsEditorModeToolkit.cpp
@@ -69,6 +69,15 @@ FModelingToolsEditorModeToolkit::FModelingToolsEditorModeToolkit()
 
 FModelingToolsEditorModeToolkit::~FModelingToolsEditorModeToolkit()
 {
+	if (IsHosted())
+	{
+		if (SelectionPaletteOverlayWidget.IsValid())
+		{
+			GetToolkitHost()->RemoveViewportOverlayWidget(SelectionPaletteOverlayWidget.ToSharedRef());
+			SelectionPaletteOverlayWidget.Reset();
+		}
+	}
+
 	UModelingToolsEditorModeSettings* Settings = GetMutableDefault<UModelingToolsEditorModeSettings>();
 	Settings->OnModified.Remove(AssetSettingsModifiedHandle);
 	GetScriptableEditorMode()->GetInteractiveToolsContext(EToolsContextScope::EdMode)->OnToolNotificationMessage.RemoveAll(this);
@@ -165,7 +174,22 @@ void FModelingToolsEditorModeToolkit::Init(const TSharedPtr<IToolkitHost>& InitT
 
 	UpdateObjectCreationOptionsFromSettings();
 
-	SAssignNew(ViewportOverlayWidget, SHorizontalBox)
+	MakeToolShutdownOverlayWidget();
+
+	const UModelingToolsEditorModeSettings* ModelingModeSettings = GetDefault<UModelingToolsEditorModeSettings>();
+	bool bEnableSelectionUI = ModelingModeSettings && ModelingModeSettings->bEnablePersistentSelections;
+	if ( bEnableSelectionUI )
+	{
+		MakeSelectionPaletteOverlayWidget();
+		GetToolkitHost()->AddViewportOverlayWidget(SelectionPaletteOverlayWidget.ToSharedRef());
+	}
+}
+
+
+
+void FModelingToolsEditorModeToolkit::MakeToolShutdownOverlayWidget()
+{
+	SAssignNew(ToolShutdownViewportOverlayWidget, SHorizontalBox)
 
 	+SHorizontalBox::Slot()
 	.HAlign(HAlign_Center)
@@ -531,8 +555,7 @@ static const FName LODToolsTabName(TEXT("LODs"));
 static const FName BakingToolsTabName(TEXT("Baking"));
 static const FName ModelingFavoritesTabName(TEXT("Favorites"));
 
-static const FName SelectionModesTabName(TEXT("SelectionModes"));
-static const FName SelectionActionsTabName(TEXT("SelectionActions"));
+static const FName SelectionActionsTabName(TEXT("Selection"));
 
 
 const TArray<FName> FModelingToolsEditorModeToolkit::PaletteNames_Standard = { PrimitiveTabName, CreateTabName, PolyModelingTabName, TriModelingTabName, DeformTabName, TransformTabName, MeshProcessingTabName, VoxToolsTabName, AttributesTabName, UVTabName, BakingToolsTabName, VolumesTabName, LODToolsTabName };
@@ -557,8 +580,6 @@ void FModelingToolsEditorModeToolkit::GetToolPaletteNames(TArray<FName>& Palette
 			PaletteNames.Insert(SelectionActionsTabName, 0);
 			ExistingNames.Add(SelectionActionsTabName);
 		}
-
-		PaletteNames.Insert(SelectionModesTabName, 0);
 	}
 
 	bool bEnablePrototypes = (CVarEnablePrototypeModelingTools.GetValueOnGameThread() > 0);
@@ -629,7 +650,7 @@ void FModelingToolsEditorModeToolkit::GetToolPaletteNames(TArray<FName>& Palette
 
 
 FText FModelingToolsEditorModeToolkit::GetToolPaletteDisplayName(FName Palette) const
-{ 
+{
 	return FText::FromName(Palette);
 }
 
@@ -655,16 +676,6 @@ void FModelingToolsEditorModeToolkit::BuildToolPalette(FName PaletteIndex, class
 			}
 		}
 	}
-	else if (PaletteIndex == SelectionModesTabName)
-	{
-		ToolbarBuilder.AddToolBarButton(Commands.BeginSelectionAction_ToObjectType);
-		ToolbarBuilder.AddToolBarButton(Commands.BeginSelectionAction_ToTriangleType);
-		ToolbarBuilder.AddToolBarButton(Commands.BeginSelectionAction_ToPolygroupType);
-
-		ToolbarBuilder.AddToolBarButton(Commands.BeginSelectionAction_ToVertexType);
-		ToolbarBuilder.AddToolBarButton(Commands.BeginSelectionAction_ToEdgeType);
-		ToolbarBuilder.AddToolBarButton(Commands.BeginSelectionAction_ToFaceType);
-	}
 	else if (PaletteIndex == SelectionActionsTabName)
 	{
 		ToolbarBuilder.AddToolBarButton(Commands.BeginSelectionAction_Delete);
@@ -1018,22 +1029,22 @@ void FModelingToolsEditorModeToolkit::InvokeUI()
 
 void FModelingToolsEditorModeToolkit::ForceToolPaletteRebuild()
 {
-
-	UGeometrySelectionManager* SelectionManager = Cast<UModelingToolsEditorMode>(GetScriptableEditorMode())->GetSelectionManager();
-	if (SelectionManager)
-	{
-		bool bHasActiveSelection = SelectionManager->HasSelection();
-
-		if (bShowActiveSelectionActions != bHasActiveSelection)
-		{
-			bShowActiveSelectionActions = bHasActiveSelection;
-			this->RebuildModeToolPalette();
-		}
-	}
-	else
-	{
-		bShowActiveSelectionActions = false;
-	}
+	// disabling this for now because it resets the expand/contract state of all the palettes...
+
+	//UGeometrySelectionManager* SelectionManager = Cast<UModelingToolsEditorMode>(GetScriptableEditorMode())->GetSelectionManager();
+	//if (SelectionManager)
+	//{
+	//	bool bHasActiveSelection = SelectionManager->HasSelection();
+	//	if (bShowActiveSelectionActions != bHasActiveSelection)
+	//	{
+	//		bShowActiveSelectionActions = bHasActiveSelection;
+	//		this->RebuildModeToolPalette();
+	//	}
+	//}
+	//else
+	//{
+	//	bShowActiveSelectionActions = false;
+	//}
 }
 
 
@@ -1084,7 +1095,7 @@ void FModelingToolsEditorModeToolkit::OnToolStarted(UInteractiveToolManager* Man
 	ActiveToolIcon = FModelingToolsEditorModeStyle::Get()->GetOptionalBrush(ActiveToolIconName);
 
 
-	GetToolkitHost()->AddViewportOverlayWidget(ViewportOverlayWidget.ToSharedRef());
+	GetToolkitHost()->AddViewportOverlayWidget(ToolShutdownViewportOverlayWidget.ToSharedRef());
 
 	// disable LOD level picker once Tool is active
 	AssetLODMode->SetEnabled(false);
@@ -1105,7 +1116,7 @@ void FModelingToolsEditorModeToolkit::OnToolEnded(UInteractiveToolManager* Manag
 {
 	if (IsHosted())
 	{
-		GetToolkitHost()->RemoveViewportOverlayWidget(ViewportOverlayWidget.ToSharedRef());
+		GetToolkitHost()->RemoveViewportOverlayWidget(ToolShutdownViewportOverlayWidget.ToSharedRef());
 	}
 
 	ModeDetailsView->SetObject(nullptr);
@@ -1134,11 +1145,11 @@ void FModelingToolsEditorModeToolkit::OnActiveViewportChanged(TSharedPtr<IAssetV
 		// Check first to see if this changed because the old viewport was deleted and if not, remove our hud
 		if (OldViewport)	
 		{
-			GetToolkitHost()->RemoveViewportOverlayWidget(ViewportOverlayWidget.ToSharedRef(), OldViewport);
+			GetToolkitHost()->RemoveViewportOverlayWidget(ToolShutdownViewportOverlayWidget.ToSharedRef(), OldViewport);
 		}
 
 		// Add the hud to the new viewport
-		GetToolkitHost()->AddViewportOverlayWidget(ViewportOverlayWidget.ToSharedRef(), NewViewport);
+		GetToolkitHost()->AddViewportOverlayWidget(ToolShutdownViewportOverlayWidget.ToSharedRef(), NewViewport);
 	}
 }
 
diff --git a/Engine/Plugins/Editor/ModelingToolsEditorMode/Source/ModelingToolsEditorMode/Private/ModelingToolsManagerActions.cpp b/Engine/Plugins/Editor/ModelingToolsEditorMode/Source/ModelingToolsEditorMode/Private/ModelingToolsManagerActions.cpp
--- a/Engine/Plugins/Editor/ModelingToolsEditorMode/Source/ModelingToolsEditorMode/Private/ModelingToolsManagerActions.cpp
+++ b/Engine/Plugins/Editor/ModelingToolsEditorMode/Source/ModelingToolsEditorMode/Private/ModelingToolsManagerActions.cpp
@@ -183,15 +183,23 @@ void FModelingToolsManagerCommands::RegisterCommands()
 	REGISTER_MODELING_TOOL_COMMAND(BeginPolyModelTool_CutFaces, "Cut", "PolyGroup Cut Faces Tool");
 
 
-	REGISTER_MODELING_TOOL_COMMAND(BeginSelectionAction_ToObjectType, "None", "Disable Geometry Selection");
+	REGISTER_MODELING_TOOL_COMMAND(BeginSelectionAction_ToObjectType, "None", "Disable Mesh Element Selection");
 	REGISTER_MODELING_TOOL_COMMAND(BeginSelectionAction_ToTriangleType, "Tris", "Set Selection Layer to Triangles");
 	REGISTER_MODELING_TOOL_COMMAND(BeginSelectionAction_ToPolygroupType, "Groups", "Set Selection Layer to Polygroups");
-
 	REGISTER_MODELING_TOOL_COMMAND(BeginSelectionAction_ToVertexType, "Vertex", "Set Selection Type to Vertices");
 	REGISTER_MODELING_TOOL_COMMAND(BeginSelectionAction_ToEdgeType, "Edge", "Set Selection Type to Edges");
 	REGISTER_MODELING_TOOL_COMMAND(BeginSelectionAction_ToFaceType, "Face", "Set Selection Type to Faces");
 
+	REGISTER_MODELING_TOOL_COMMAND(MeshSelectionModeAction_NoSelection, "None", "Disable Geometry Selection");
+	REGISTER_MODELING_TOOL_COMMAND(MeshSelectionModeAction_MeshTriangles, "Tris", "Select Mesh Triangles");
+	REGISTER_MODELING_TOOL_COMMAND(MeshSelectionModeAction_MeshVertices, "Verts", "Select Mesh Vertices");
+	REGISTER_MODELING_TOOL_COMMAND(MeshSelectionModeAction_MeshEdges, "Edges", "Select Mesh Edges");
+	REGISTER_MODELING_TOOL_COMMAND(MeshSelectionModeAction_GroupFaces, "Groups", "Select Mesh Polygroups");
+	REGISTER_MODELING_TOOL_COMMAND(MeshSelectionModeAction_GroupCorners, "Corners", "Select Mesh Polygroup Corners/Vertices");
+	REGISTER_MODELING_TOOL_COMMAND(MeshSelectionModeAction_GroupEdges, "Borders", "Select Mesh Polygroup Borders/Edges");
+
 	REGISTER_MODELING_TOOL_COMMAND(BeginSelectionAction_Delete, "Delete", "Delete Selection");
+	REGISTER_MODELING_TOOL_COMMAND(BeginSelectionAction_Extrude, "Extrude", "Extrude Selection");
 
 	UI_COMMAND(AcceptActiveTool, "Accept", "Accept the active tool", EUserInterfaceActionType::Button, FInputChord());
 	UI_COMMAND(CancelActiveTool, "Cancel", "Cancel the active tool", EUserInterfaceActionType::Button, FInputChord());
diff --git a/Engine/Plugins/Editor/ModelingToolsEditorMode/Source/ModelingToolsEditorMode/Public/ModelingSelectionInteraction.h b/Engine/Plugins/Editor/ModelingToolsEditorMode/Source/ModelingToolsEditorMode/Public/ModelingSelectionInteraction.h
--- a/Engine/Plugins/Editor/ModelingToolsEditorMode/Source/ModelingToolsEditorMode/Public/ModelingSelectionInteraction.h
+++ b/Engine/Plugins/Editor/ModelingToolsEditorMode/Source/ModelingToolsEditorMode/Public/ModelingSelectionInteraction.h
@@ -7,11 +7,28 @@
 #include "InputBehaviorSet.h"
 #include "BaseGizmos/TransformGizmoUtil.h"
 #include "FrameTypes.h"
+#include "Mechanics/RectangleMarqueeMechanic.h"
+#include "Mechanics/PolyLassoMarqueeMechanic.h"
+#include "Selections/GeometrySelection.h"
 #include "ModelingSelectionInteraction.generated.h"
 
 class UGeometrySelectionManager;
 class USingleClickInputBehavior;
+class USingleClickOrDragInputBehavior;
 class UMouseHoverBehavior;
+class URectangleMarqueeMechanic;
+class UPathSelectionInteraction;
+class FCanvas;
+class IToolsContextRenderAPI;
+
+
+UENUM()
+enum class EModelingSelectionInteraction_DragMode : uint8
+{
+	NoDragInteraction = 0,
+	PathInteraction = 1,
+	RectangleMarqueeInteraction = 2
+};
 
 /**
  * UModelingSelectionInteraction provides element-level selection behavior (ie mesh triangles/edges/vertices)
@@ -49,6 +66,23 @@ public:
 
 	virtual void Shutdown();
 
+	UGeometrySelectionManager* GetSelectionManager() { return SelectionManager; }
+	UE::Geometry::FGeometrySelectionUpdateConfig GetActiveSelectionUpdateConfig() const;
+
+	virtual void DrawHUD(FCanvas* Canvas, IToolsContextRenderAPI* RenderAPI);
+
+
+	//
+	// Support for drag-style interactions
+	//
+
+	EModelingSelectionInteraction_DragMode GetActiveDragMode() const { return ActiveDragMode; }
+	void SetActiveDragMode(EModelingSelectionInteraction_DragMode NewMode);
+
+	//
+	// Support for transformation of selected geometry
+	//
+
 	// this needs to be called from EdMode tick or rendering, to avoid mouse-interrupt priority issues
 	virtual void ApplyPendingTransformInteractions();
 
@@ -58,10 +92,11 @@ public:
 	DECLARE_MULTICAST_DELEGATE(FOnEndTransformInteraction);
 	FOnEndTransformInteraction OnTransformEnd;
 
-public:
-	// click-to-select behavior
+
+protected:
+	// click-to-select behavior, various drag behaviors
 	UPROPERTY()
-	TObjectPtr<USingleClickInputBehavior> ClickBehavior;
+	TObjectPtr<USingleClickOrDragInputBehavior> ClickOrDragBehavior;
 
 	// mouse hover behavior
 	UPROPERTY()
@@ -152,4 +187,43 @@ protected:
 	void OnEndGizmoTransform(UTransformProxy* Proxy);
 	void OnGizmoTransformChanged(UTransformProxy* Proxy, FTransform Transform);
 
+
+protected:
+	EModelingSelectionInteraction_DragMode ActiveDragMode = EModelingSelectionInteraction_DragMode::PathInteraction;
+	void UpdateActiveDragMode();
+
+	UPROPERTY()
+	TObjectPtr<URectangleMarqueeInteraction> RectangleMarqueeInteraction;
+
+	void OnMarqueeRectangleFinished(const FCameraRectangle& Rectangle, bool bCancelled);
+
+	UPROPERTY()
+	TObjectPtr<UPathSelectionInteraction> PathSelectionInteraction;
+};
+
+
+
+/**
+ * UPathSelectionInteraction is a simple drag-interaction for Selection, which essentially just
+ * selects/deselects any element hit by the cursor. 
+ * 
+ * Currently no attempt is made to (eg) subsample, so a fast-moving cursor will skip over some faces.
+ */
+UCLASS()
+class UPathSelectionInteraction : public UObject, public IClickDragBehaviorTarget
+{
+	GENERATED_BODY()
+public:
+	void Setup(UModelingSelectionInteraction* SelectionInteraction);
+
+protected:
+	UPROPERTY()
+	TWeakObjectPtr<UModelingSelectionInteraction> SelectionInteraction = nullptr;
+
+	virtual FInputRayHit CanBeginClickDragSequence(const FInputDeviceRay& PressPos) override;
+	virtual void OnClickPress(const FInputDeviceRay& PressPos) override;
+	virtual void OnClickDrag(const FInputDeviceRay& DragPos) override;
+	virtual void OnClickRelease(const FInputDeviceRay& ReleasePos) override;
+	virtual void OnTerminateDragSequence() override;
+
 };
\ No newline at end of file
diff --git a/Engine/Plugins/Editor/ModelingToolsEditorMode/Source/ModelingToolsEditorMode/Public/ModelingToolsEditorMode.h b/Engine/Plugins/Editor/ModelingToolsEditorMode/Source/ModelingToolsEditorMode/Public/ModelingToolsEditorMode.h
--- a/Engine/Plugins/Editor/ModelingToolsEditorMode/Source/ModelingToolsEditorMode/Public/ModelingToolsEditorMode.h
+++ b/Engine/Plugins/Editor/ModelingToolsEditorMode/Source/ModelingToolsEditorMode/Public/ModelingToolsEditorMode.h
@@ -66,10 +66,25 @@ public:
 	//////////////////
 
 
+	//
+	// Selection System configuration, this will likely move elsewhere
+	//
+
 	virtual UGeometrySelectionManager* GetSelectionManager() const
 	{
 		return SelectionManager;
 	}
+	virtual UModelingSelectionInteraction* GetSelectionInteraction() const
+	{
+		return SelectionInteraction;
+	}
+
+	UPROPERTY()
+	bool bEnableVolumeElementSelection = false;
+
+	UPROPERTY()
+	bool bEnableStaticMeshElementSelection = false;
+
 
 protected:
 	virtual void BindCommands() override;
@@ -104,9 +119,10 @@ protected:
 	bool GetGeometrySelectionChangesAllowed() const;
 	bool TestForEditorGizmoHit(const FInputDeviceRay&) const;
 
-	void UpdateSelectionManagerOnEditorSelectionChange();
+	void UpdateSelectionManagerOnEditorSelectionChange(bool bEnteringMode = false);
 
 	void OnToolsContextRender(IToolsContextRenderAPI* RenderAPI);
+	void OnToolsContextDrawHUD(FCanvas* Canvas, IToolsContextRenderAPI* RenderAPI);
 
 	void ModelingModeShortcutRequested(EModelingModeActionCommands Command);
 	void FocusCameraAtCursorHotkey();
diff --git a/Engine/Plugins/Editor/ModelingToolsEditorMode/Source/ModelingToolsEditorMode/Public/ModelingToolsEditorModeSettings.h b/Engine/Plugins/Editor/ModelingToolsEditorMode/Source/ModelingToolsEditorMode/Public/ModelingToolsEditorModeSettings.h
--- a/Engine/Plugins/Editor/ModelingToolsEditorMode/Source/ModelingToolsEditorMode/Public/ModelingToolsEditorModeSettings.h
+++ b/Engine/Plugins/Editor/ModelingToolsEditorMode/Source/ModelingToolsEditorMode/Public/ModelingToolsEditorModeSettings.h
@@ -145,8 +145,8 @@ public:
 	// Selection
 	//
 
-	/** If true, selections will be tracked between tools, and visualized with outlines. This feature is Experimental, disabled by default. */
-	UPROPERTY(config, EditAnywhere, Category = "Modeling Mode|Selection")
+	/** If true, Mesh Selection System (Experimental) will be enabled. */
+	UPROPERTY(config, EditAnywhere, Category = "Modeling Mode|Selection", meta = (DisplayName="Enable Mesh Selection UI"))
 	bool bEnablePersistentSelections = false;
 
 	DECLARE_MULTICAST_DELEGATE_TwoParams(UModelingToolsEditorModeSettingsModified, UObject*, FProperty*);
diff --git a/Engine/Plugins/Editor/ModelingToolsEditorMode/Source/ModelingToolsEditorMode/Public/ModelingToolsEditorModeToolkit.h b/Engine/Plugins/Editor/ModelingToolsEditorMode/Source/ModelingToolsEditorMode/Public/ModelingToolsEditorModeToolkit.h
--- a/Engine/Plugins/Editor/ModelingToolsEditorMode/Source/ModelingToolsEditorMode/Public/ModelingToolsEditorModeToolkit.h
+++ b/Engine/Plugins/Editor/ModelingToolsEditorMode/Source/ModelingToolsEditorMode/Public/ModelingToolsEditorModeToolkit.h
@@ -80,7 +80,12 @@ private:
 	void UpdateActiveToolProperties();
 	void InvalidateCachedDetailPanelState(UObject* ChangedObject);
 
-	TSharedPtr<SWidget> ViewportOverlayWidget;
+	TSharedPtr<SWidget> ToolShutdownViewportOverlayWidget;
+	void MakeToolShutdownOverlayWidget();
+
+	TSharedPtr<SWidget> SelectionPaletteOverlayWidget;
+	void MakeSelectionPaletteOverlayWidget();
+
 
 	TSharedPtr<STextBlock> ModeWarningArea;
 	TSharedPtr<STextBlock> ModeHeaderArea;
@@ -112,10 +117,5 @@ private:
 	TSharedPtr<STextComboBox> AssetLODMode;
 
 	bool bFirstInitializeAfterModeSetup = true;
-
-
-	bool bShowActiveSelectionActions = false;
-
-
-	// custom accept/cancel/complete handlers
+	bool bShowActiveSelectionActions = true;
 };
diff --git a/Engine/Plugins/Editor/ModelingToolsEditorMode/Source/ModelingToolsEditorMode/Public/ModelingToolsManagerActions.h b/Engine/Plugins/Editor/ModelingToolsEditorMode/Source/ModelingToolsEditorMode/Public/ModelingToolsManagerActions.h
--- a/Engine/Plugins/Editor/ModelingToolsEditorMode/Source/ModelingToolsEditorMode/Public/ModelingToolsManagerActions.h
+++ b/Engine/Plugins/Editor/ModelingToolsEditorMode/Source/ModelingToolsEditorMode/Public/ModelingToolsManagerActions.h
@@ -145,15 +145,25 @@ public:
 	TSharedPtr<FUICommandInfo> BeginPolyModelTool_CutFaces;
 
 
+	// old selection state actions
 	TSharedPtr<FUICommandInfo> BeginSelectionAction_ToVertexType;
 	TSharedPtr<FUICommandInfo> BeginSelectionAction_ToEdgeType;
 	TSharedPtr<FUICommandInfo> BeginSelectionAction_ToFaceType;
-
 	TSharedPtr<FUICommandInfo> BeginSelectionAction_ToObjectType;
 	TSharedPtr<FUICommandInfo> BeginSelectionAction_ToTriangleType;
 	TSharedPtr<FUICommandInfo> BeginSelectionAction_ToPolygroupType;
 
+	TSharedPtr<FUICommandInfo> MeshSelectionModeAction_NoSelection;
+	TSharedPtr<FUICommandInfo> MeshSelectionModeAction_MeshTriangles;
+	TSharedPtr<FUICommandInfo> MeshSelectionModeAction_MeshVertices;
+	TSharedPtr<FUICommandInfo> MeshSelectionModeAction_MeshEdges;
+	TSharedPtr<FUICommandInfo> MeshSelectionModeAction_GroupFaces;
+	TSharedPtr<FUICommandInfo> MeshSelectionModeAction_GroupCorners;
+	TSharedPtr<FUICommandInfo> MeshSelectionModeAction_GroupEdges;
+
+
 	TSharedPtr<FUICommandInfo> BeginSelectionAction_Delete;
+	TSharedPtr<FUICommandInfo> BeginSelectionAction_Extrude;
 
 	// For connecting ModelingMode with UVEditor plugin
 	TSharedPtr<FUICommandInfo> LaunchUVEditor;
diff --git a/Engine/Plugins/Runtime/MeshModelingToolset/Source/ModelingComponents/Private/Mechanics/RectangleMarqueeMechanic.cpp b/Engine/Plugins/Runtime/MeshModelingToolset/Source/ModelingComponents/Private/Mechanics/RectangleMarqueeMechanic.cpp
--- a/Engine/Plugins/Runtime/MeshModelingToolset/Source/ModelingComponents/Private/Mechanics/RectangleMarqueeMechanic.cpp
+++ b/Engine/Plugins/Runtime/MeshModelingToolset/Source/ModelingComponents/Private/Mechanics/RectangleMarqueeMechanic.cpp
@@ -180,6 +180,85 @@ FConvexVolume FCameraRectangle::FrustumAsConvexVolume() const
 	return OutFrustum;
 }
 
+
+
+
+
+// ---------------------------------------
+
+FInputRayHit URectangleMarqueeInteraction::CanBeginClickDragSequence(const FInputDeviceRay& PressPos)
+{
+	return FInputRayHit(TNumericLimits<float>::Max());
+}
+
+void URectangleMarqueeInteraction::OnClickPress(const FInputDeviceRay& PressPos)
+{
+	if (!PressPos.bHas2D)
+	{
+		bIsDragging = false;
+		return;
+	}
+
+	CameraRectangle.RectangleStartRay = PressPos;
+	CameraRectangle.RectangleEndRay = PressPos;
+	CameraRectangle.Initialize();
+
+	OnDragRectangleStarted.Broadcast();
+	// consider broadcasting OnDragRectangleChanged as well here, and set bIsDragging=true here...
+}
+
+void URectangleMarqueeInteraction::OnClickDrag(const FInputDeviceRay& DragPos)
+{
+	if (!DragPos.bHas2D)
+	{
+		return;
+	}
+
+	// should probably threshold based on movement here, ie if 2D position is the same, do not 
+	// need to post an update... (no guarantees that this isn't called continuously w/ same DragPos!)
+
+	bIsDragging = true;
+	CameraRectangle.RectangleEndRay = DragPos;
+	CameraRectangle.Initialize();
+	
+	OnDragRectangleChanged.Broadcast(CameraRectangle);
+}
+
+void URectangleMarqueeInteraction::OnClickRelease(const FInputDeviceRay& ReleasePos)
+{
+	bIsDragging = false;
+	OnDragRectangleFinished.Broadcast(CameraRectangle, false);
+}
+
+void URectangleMarqueeInteraction::OnTerminateDragSequence()
+{
+	bIsDragging = false;
+	OnDragRectangleFinished.Broadcast(CameraRectangle, true);
+}
+
+
+void URectangleMarqueeInteraction::DrawHUD(FCanvas* Canvas, IToolsContextRenderAPI* RenderAPI)
+{
+	CurrentCameraState = RenderAPI->GetCameraState();
+
+	EViewInteractionState State = RenderAPI->GetViewInteractionState();
+	bool bThisViewHasFocus = !!(State & EViewInteractionState::Focused);
+	if (bThisViewHasFocus && bIsDragging)
+	{
+		FVector2D Start = CameraRectangle.RectangleStartRay.ScreenPosition;
+		FVector2D Curr = CameraRectangle.RectangleEndRay.ScreenPosition;
+		FCanvasBoxItem BoxItem(Start / Canvas->GetDPIScale(), (Curr - Start) / Canvas->GetDPIScale());
+		BoxItem.SetColor(FLinearColor::White);
+		Canvas->DrawItem(BoxItem);
+	}
+}
+
+
+
+
+
+
+
 // ---------------------------------------
 
 void URectangleMarqueeMechanic::Setup(UInteractiveTool* ParentToolIn)
diff --git a/Engine/Plugins/Runtime/MeshModelingToolset/Source/ModelingComponents/Public/Mechanics/RectangleMarqueeMechanic.h b/Engine/Plugins/Runtime/MeshModelingToolset/Source/ModelingComponents/Public/Mechanics/RectangleMarqueeMechanic.h
--- a/Engine/Plugins/Runtime/MeshModelingToolset/Source/ModelingComponents/Public/Mechanics/RectangleMarqueeMechanic.h
+++ b/Engine/Plugins/Runtime/MeshModelingToolset/Source/ModelingComponents/Public/Mechanics/RectangleMarqueeMechanic.h
@@ -118,7 +118,6 @@ struct MODELINGCOMPONENTS_API FCameraRectangle
  *
  * Attach to the mechanic's delegates and use the passed rectangle to test against your geometry. 
  */
-
 UCLASS()
 class MODELINGCOMPONENTS_API URectangleMarqueeMechanic : public UInteractionMechanic, public IClickDragBehaviorTarget
 {
@@ -222,3 +221,54 @@ private:
 	bool bIsDragging;
 	bool bIsOnDragRectangleChangedDeferred;
 };
+
+
+
+
+
+/*
+ * URectangleMarqueeInteraction is a simplified version of URectangleMarqueeMechanic that is not a UInteractionMechanic,
+ * which requires a base Tool/ToolManager. This variant does not create it's own InputBehavior, but is still a 
+ * ClickDragBehaviorTarget. 
+ * 
+ * DrawHUD() must be called by the owning code. 
+ */
+UCLASS()
+class MODELINGCOMPONENTS_API URectangleMarqueeInteraction : public UObject, public IClickDragBehaviorTarget
+{
+	GENERATED_BODY()
+
+public:
+
+	virtual void DrawHUD(FCanvas* Canvas, IToolsContextRenderAPI* RenderAPI);
+
+	/**
+	 * OnDragRectangleStarted is called when user starts dragging a new rectangle.
+	 */
+	FSimpleMulticastDelegate OnDragRectangleStarted;
+
+	/**
+	 * OnDragRectangleChanged is called as the user drags the other corner of the rectangle around.
+	 */
+	DECLARE_MULTICAST_DELEGATE_OneParam(OnDragRectangleChangedEvent, const FCameraRectangle&);
+	OnDragRectangleChangedEvent OnDragRectangleChanged;
+
+	/**
+	 * OnDragRectangleFinished is called once the user lets go of the mouse button after dragging out a rectangle.
+	 * bCancelled flag is true when the drag finishes due to a disabling of the mechanic or due to a TerminateDragSequence call, rather than a normal drag completion.
+	 */
+	DECLARE_MULTICAST_DELEGATE_TwoParams(OnDragRectangleFinishedEvent, const FCameraRectangle&, bool bCancelled);
+	OnDragRectangleFinishedEvent OnDragRectangleFinished;
+	
+protected:
+	FViewCameraState CurrentCameraState;
+	FCameraRectangle CameraRectangle;
+
+	virtual FInputRayHit CanBeginClickDragSequence(const FInputDeviceRay& PressPos) override;
+	virtual void OnClickPress(const FInputDeviceRay& PressPos) override;
+	virtual void OnClickDrag(const FInputDeviceRay& DragPos) override;
+	virtual void OnClickRelease(const FInputDeviceRay& ReleasePos) override;
+	virtual void OnTerminateDragSequence() override;
+
+	bool bIsDragging = false;
+};
diff --git a/Engine/Source/Runtime/InteractiveToolsFramework/Private/BaseBehaviors/SingleClickOrDragBehavior.cpp b/Engine/Source/Runtime/InteractiveToolsFramework/Private/BaseBehaviors/SingleClickOrDragBehavior.cpp
--- a/Engine/Source/Runtime/InteractiveToolsFramework/Private/BaseBehaviors/SingleClickOrDragBehavior.cpp
+++ b/Engine/Source/Runtime/InteractiveToolsFramework/Private/BaseBehaviors/SingleClickOrDragBehavior.cpp
@@ -12,13 +12,21 @@ void USingleClickOrDragInputBehavior::Initialize(IClickBehaviorTarget* ClickTarg
 	this->DragTarget = DragTargetIn;
 }
 
+void USingleClickOrDragInputBehavior::SetDragTarget(IClickDragBehaviorTarget* DragTargetIn)
+{
+	if ( ensure(!bInDrag) )
+	{
+		this->DragTarget = DragTargetIn;
+	}
+}
+
 
 FInputCaptureRequest USingleClickOrDragInputBehavior::WantsCapture(const FInputDeviceState& Input)
 {
 	bImmediatelyBeginDragInBeginCapture = false;
 	if (IsPressed(Input) && (ModifierCheckFunc == nullptr || ModifierCheckFunc(Input)) )
 	{
-		FInputRayHit ClickHitResult = ClickTarget->IsHitByClick(GetDeviceRay(Input));
+		FInputRayHit ClickHitResult = (ClickTarget) ? ClickTarget->IsHitByClick(GetDeviceRay(Input)) : FInputRayHit();
 		if (ClickHitResult.bHit)
 		{
 			return FInputCaptureRequest::Begin(this, EInputCaptureSide::Any, ClickHitResult.HitDepth);
@@ -39,8 +47,14 @@ FInputCaptureRequest USingleClickOrDragInputBehavior::WantsCapture(const FInputD
 
 FInputCaptureUpdate USingleClickOrDragInputBehavior::BeginCapture(const FInputDeviceState& Input, EInputCaptureSide eSide)
 {
-	Modifiers.UpdateModifiers(Input, ClickTarget);
-	Modifiers.UpdateModifiers(Input, DragTarget);
+	if ( ClickTarget != nullptr )
+	{
+		Modifiers.UpdateModifiers(Input, ClickTarget);
+	}
+	if ( DragTarget != nullptr )
+	{
+		Modifiers.UpdateModifiers(Input, DragTarget);
+	}
 
 	ensure(Input.IsFromDevice(EInputDevices::Mouse));	// todo: handle other devices
 	MouseDownPosition2D = Input.Mouse.Position2D;	
@@ -62,8 +76,14 @@ FInputCaptureUpdate USingleClickOrDragInputBehavior::BeginCapture(const FInputDe
 
 FInputCaptureUpdate USingleClickOrDragInputBehavior::UpdateCapture(const FInputDeviceState& Input, const FInputCaptureData& Data)
 {
-	Modifiers.UpdateModifiers(Input, ClickTarget);
-	Modifiers.UpdateModifiers(Input, DragTarget);
+	if ( ClickTarget != nullptr )
+	{
+		Modifiers.UpdateModifiers(Input, ClickTarget);
+	}
+	if ( DragTarget != nullptr )
+	{
+		Modifiers.UpdateModifiers(Input, DragTarget);
+	}
 
 	// check if mouse has moved far enough that we want to swap to drag behavior
 	if (bInDrag == false)
@@ -74,11 +94,14 @@ FInputCaptureUpdate USingleClickOrDragInputBehavior::UpdateCapture(const FInputD
 			FInputDeviceState StartInput = Input;
 			StartInput.Mouse.Position2D = MouseDownPosition2D;		// attempt to reconstruct the input state that would have existed in WantsCapture/BeginCapture
 			StartInput.Mouse.WorldRay = MouseDownRay;
-			FInputRayHit DragHitResult = DragTarget->CanBeginClickDragSequence(GetDeviceRay(StartInput));
-			if (DragHitResult.bHit)
+			if ( DragTarget )
 			{
-				bInDrag = true;
-				OnClickDragPressInternal(StartInput, CaptureSide);
+				FInputRayHit DragHitResult = DragTarget->CanBeginClickDragSequence(GetDeviceRay(StartInput));
+				if (DragHitResult.bHit)
+				{
+					bInDrag = true;
+					OnClickDragPressInternal(StartInput, CaptureSide);
+				}
 			}
 		}
 	}
@@ -92,7 +115,7 @@ FInputCaptureUpdate USingleClickOrDragInputBehavior::UpdateCapture(const FInputD
 		}
 		else   // click path
 		{
-			if (ClickTarget->IsHitByClick(GetDeviceRay(Input)).bHit)
+			if (ClickTarget && ClickTarget->IsHitByClick(GetDeviceRay(Input)).bHit )
 			{
 				OnClickedInternal(Input, Data);
 			}
@@ -115,7 +138,10 @@ void USingleClickOrDragInputBehavior::ForceEndCapture(const FInputCaptureData& d
 {
 	if (bInDrag)
 	{
-		DragTarget->OnTerminateDragSequence();
+		if ( DragTarget )
+		{
+			DragTarget->OnTerminateDragSequence();
+		}
 		bInDrag = false;
 	}
 
@@ -125,23 +151,35 @@ void USingleClickOrDragInputBehavior::ForceEndCapture(const FInputCaptureData& d
 
 void USingleClickOrDragInputBehavior::OnClickedInternal(const FInputDeviceState& input, const FInputCaptureData& data)
 {
-	ClickTarget->OnClicked(GetDeviceRay(input));
+	if ( ClickTarget )
+	{
+		ClickTarget->OnClicked(GetDeviceRay(input));
+	}
 }
 
 
 
 void USingleClickOrDragInputBehavior::OnClickDragPressInternal(const FInputDeviceState& Input, EInputCaptureSide Side)
 {
-	DragTarget->OnClickPress(GetDeviceRay(Input));
+	if ( DragTarget )
+	{
+		DragTarget->OnClickPress(GetDeviceRay(Input));
+	}
 }
 
 void USingleClickOrDragInputBehavior::OnClickDragInternal(const FInputDeviceState& Input, const FInputCaptureData& Data)
 {
-	DragTarget->OnClickDrag(GetDeviceRay(Input));
+	if ( DragTarget )
+	{
+		DragTarget->OnClickDrag(GetDeviceRay(Input));
+	}
 }
 
 void USingleClickOrDragInputBehavior::OnClickDragReleaseInternal(const FInputDeviceState& Input, const FInputCaptureData& Data)
 {
-	DragTarget->OnClickRelease(GetDeviceRay(Input));
+	if ( DragTarget )
+	{
+		DragTarget->OnClickRelease(GetDeviceRay(Input));
+	}
 }
 
diff --git a/Engine/Source/Runtime/InteractiveToolsFramework/Public/BaseBehaviors/SingleClickOrDragBehavior.h b/Engine/Source/Runtime/InteractiveToolsFramework/Public/BaseBehaviors/SingleClickOrDragBehavior.h
--- a/Engine/Source/Runtime/InteractiveToolsFramework/Public/BaseBehaviors/SingleClickOrDragBehavior.h
+++ b/Engine/Source/Runtime/InteractiveToolsFramework/Public/BaseBehaviors/SingleClickOrDragBehavior.h
@@ -55,6 +55,11 @@ public:
 	 */
 	virtual void Initialize(IClickBehaviorTarget* ClickTarget, IClickDragBehaviorTarget* DragTarget);
 
+	/**
+	 * Change the Drag BehaviorTarget
+	 */
+	virtual void SetDragTarget(IClickDragBehaviorTarget* DragTarget);
+
 
 	/**
 	 * WantsCapture() will only return capture request if this function returns true (or is null)
