ModelingMode: disable slate throttling when we are inside a Tool. Currently slate interaction throttling (intended to improve editor responsiveness) disables editor mode ticking when the user is dragging a slider in a details panel (eg like occurs in nearly all modeling tool panels). However if a modeling mode tool is computing results in a background thread, and reading back on Tick(), this means that the visible result will always be one step behind until the user releases the mouse, no matter how long they wait (ie, the opposite of responsiveness). This results in many bad interactions, like the user stopping dragging a slider when they are happy with what they see, only to have the result change when they release the mouse.

diff --git a/Engine/Plugins/Editor/ModelingToolsEditorMode/Source/ModelingToolsEditorMode/Private/ModelingToolsEditorMode.cpp b/Engine/Plugins/Editor/ModelingToolsEditorMode/Source/ModelingToolsEditorMode/Private/ModelingToolsEditorMode.cpp
--- a/Engine/Plugins/Editor/ModelingToolsEditorMode/Source/ModelingToolsEditorMode/Private/ModelingToolsEditorMode.cpp
+++ b/Engine/Plugins/Editor/ModelingToolsEditorMode/Source/ModelingToolsEditorMode/Private/ModelingToolsEditorMode.cpp
@@ -119,6 +119,7 @@
 #include "LevelEditor.h"
 #include "IAssetViewport.h"
 #include "SLevelViewport.h"
+#include "Application/ThrottleManager.h"
 
 #include "ModelingToolsActions.h"
 #include "ModelingToolsManagerActions.h"
@@ -1226,6 +1227,12 @@ void UModelingToolsEditorMode::OnToolPostBuild(
 
 void UModelingToolsEditorMode::OnToolStarted(UInteractiveToolManager* Manager, UInteractiveTool* Tool)
 {
+	// disable slate throttling so that Tool background computes responding to sliders can properly be processed
+	// on Tool Tick. Otherwise, when a Tool kicks off a background update in a background thread, the computed
+	// result will be ignored until the user moves the slider, ie you cannot hold down the mouse and wait to see
+	// the result. This apparently broken behavior is currently by-design.
+	FSlateThrottleManager::Get().DisableThrottle(true);
+
 	FModelingToolActionCommands::UpdateToolCommandBinding(Tool, Toolkit->GetToolkitCommands(), false);
 	
 	if( FEngineAnalytics::IsAvailable() )
@@ -1238,6 +1245,9 @@ void UModelingToolsEditorMode::OnToolStarted(UInteractiveToolManager* Manager, U
 
 void UModelingToolsEditorMode::OnToolEnded(UInteractiveToolManager* Manager, UInteractiveTool* Tool)
 {
+	// re-enable slate throttling (see OnToolStarted)
+	FSlateThrottleManager::Get().DisableThrottle(false);
+
 	FModelingToolActionCommands::UpdateToolCommandBinding(Tool, Toolkit->GetToolkitCommands(), true);
 	
 	if( FEngineAnalytics::IsAvailable() )
diff --git a/Engine/Plugins/Experimental/ScriptableToolsEditorMode/Source/ScriptableToolsEditorMode/Private/ScriptableToolsEditorMode.cpp b/Engine/Plugins/Experimental/ScriptableToolsEditorMode/Source/ScriptableToolsEditorMode/Private/ScriptableToolsEditorMode.cpp
--- a/Engine/Plugins/Experimental/ScriptableToolsEditorMode/Source/ScriptableToolsEditorMode/Private/ScriptableToolsEditorMode.cpp
+++ b/Engine/Plugins/Experimental/ScriptableToolsEditorMode/Source/ScriptableToolsEditorMode/Private/ScriptableToolsEditorMode.cpp
@@ -10,6 +10,7 @@
 #include "EditorModeManager.h"
 #include "UnrealWidget.h"
 #include "Framework/Commands/UICommandList.h"
+#include "Application/ThrottleManager.h"
 
 #include "ScriptableToolsEditorModeToolkit.h"
 #include "ScriptableToolsEditorModeManagerCommands.h"
@@ -262,10 +263,17 @@ void UScriptableToolsEditorMode::OnToolPostBuild(
 
 void UScriptableToolsEditorMode::OnToolStarted(UInteractiveToolManager* Manager, UInteractiveTool* Tool)
 {
+	// disable slate throttling so that Tool background computes responding to sliders can properly be processed
+	// on Tool Tick. Otherwise, when a Tool kicks off a background update in a background thread, the computed
+	// result will be ignored until the user moves the slider, ie you cannot hold down the mouse and wait to see
+	// the result. This apparently broken behavior is currently by-design.
+	FSlateThrottleManager::Get().DisableThrottle(true);
 }
 
 void UScriptableToolsEditorMode::OnToolEnded(UInteractiveToolManager* Manager, UInteractiveTool* Tool)
 {
+	// re-enable slate throttling (see OnToolStarted)
+	FSlateThrottleManager::Get().DisableThrottle(false);
 }
 
 void UScriptableToolsEditorMode::BindCommands()
