[Engine/Plugins/Editor]

diff --git a/Engine/Plugins/Editor/GameplayTagsEditor/Source/GameplayTagsEditor/Private/SGameplayTagWidget.cpp b/Engine/Plugins/Editor/GameplayTagsEditor/Source/GameplayTagsEditor/Private/SGameplayTagWidget.cpp
--- a/Engine/Plugins/Editor/GameplayTagsEditor/Source/GameplayTagsEditor/Private/SGameplayTagWidget.cpp
+++ b/Engine/Plugins/Editor/GameplayTagsEditor/Source/GameplayTagsEditor/Private/SGameplayTagWidget.cpp
@@ -1,39 +1,30 @@
 // Copyright Epic Games, Inc. All Rights Reserved.
 
 #include "SGameplayTagWidget.h"
+#include "AssetRegistry/AssetIdentifier.h"
 #include "Misc/ConfigCacheIni.h"
+#include "Framework/Views/TableViewMetadata.h"
 #include "Widgets/Input/SButton.h"
+#include "GameplayTagsManager.h"
 #include "Widgets/Input/SCheckBox.h"
+#include "Widgets/IToolTip.h"
 #include "Widgets/Input/SComboButton.h"
-#include "Widgets/Images/SImage.h"
-#include "Styling/AppStyle.h"
 #include "Widgets/SWindow.h"
 #include "Misc/MessageDialog.h"
-#include "GameplayTagsModule.h"
 #include "ScopedTransaction.h"
-#include "Textures/SlateIcon.h"
 #include "PropertyHandle.h"
 #include "Widgets/Input/SSearchBox.h"
 #include "GameplayTagsEditorModule.h"
-#include "Widgets/Layout/SScaleBox.h"
-
-#include "AssetToolsModule.h"
-#include "Widgets/Input/SHyperlink.h"
-#include "Widgets/Notifications/SNotificationList.h"
-#include "Framework/Notifications/NotificationManager.h"
-#include "GameplayTagsSettings.h"
-#include "Layout/WidgetPath.h"
+
 #include "Framework/Application/SlateApplication.h"
 #include "SAddNewGameplayTagWidget.h"
 #include "SAddNewGameplayTagSourceWidget.h"
 #include "SAddNewRestrictedGameplayTagWidget.h"
 #include "SRenameGameplayTagDialog.h"
-#include "AssetRegistry/AssetData.h"
 #include "Editor.h"
-#include "Framework/Commands/UIAction.h"
 #include "Framework/MultiBox/MultiBoxBuilder.h"
-#include "Brushes/SlateColorBrush.h"
 #include "Interfaces/IMainFrameModule.h"
+#include "Widgets/Views/STreeView.h"
 
 #define LOCTEXT_NAMESPACE "GameplayTagWidget"
 
diff --git a/Engine/Plugins/Editor/GameplayTagsEditor/Source/GameplayTagsEditor/Private/SGraphNode_MultiCompareGameplayTag.h b/Engine/Plugins/Editor/GameplayTagsEditor/Source/GameplayTagsEditor/Private/SGraphNode_MultiCompareGameplayTag.h
--- a/Engine/Plugins/Editor/GameplayTagsEditor/Source/GameplayTagsEditor/Private/SGraphNode_MultiCompareGameplayTag.h
+++ b/Engine/Plugins/Editor/GameplayTagsEditor/Source/GameplayTagsEditor/Private/SGraphNode_MultiCompareGameplayTag.h
@@ -2,10 +2,6 @@
 
 #pragma once
 
-#include "CoreMinimal.h"
-#include "Layout/Visibility.h"
-#include "Input/Reply.h"
-#include "Widgets/DeclarativeSyntaxSupport.h"
 #include "SGraphNode.h"
 
 class SVerticalBox;
diff --git a/Engine/Plugins/Editor/GameplayTagsEditor/Source/GameplayTagsEditor/Private/SRenameGameplayTagDialog.cpp b/Engine/Plugins/Editor/GameplayTagsEditor/Source/GameplayTagsEditor/Private/SRenameGameplayTagDialog.cpp
--- a/Engine/Plugins/Editor/GameplayTagsEditor/Source/GameplayTagsEditor/Private/SRenameGameplayTagDialog.cpp
+++ b/Engine/Plugins/Editor/GameplayTagsEditor/Source/GameplayTagsEditor/Private/SRenameGameplayTagDialog.cpp
@@ -2,11 +2,12 @@
 
 #include "SRenameGameplayTagDialog.h"
 #include "GameplayTagsEditorModule.h"
+#include "GameplayTagsManager.h"
 #include "Widgets/Input/SButton.h"
-#include "Widgets/SBoxPanel.h"
+#include "Widgets/Input/SEditableTextBox.h"
 #include "Widgets/Text/STextBlock.h"
-#include "Layout/WidgetPath.h"
 #include "Framework/Application/SlateApplication.h"
+#include "Widgets/SWindow.h"
 
 #define LOCTEXT_NAMESPACE "RenameGameplayTag"
 
diff --git a/Engine/Plugins/Editor/GameplayTagsEditor/Source/GameplayTagsEditor/Private/SRenameGameplayTagDialog.h b/Engine/Plugins/Editor/GameplayTagsEditor/Source/GameplayTagsEditor/Private/SRenameGameplayTagDialog.h
--- a/Engine/Plugins/Editor/GameplayTagsEditor/Source/GameplayTagsEditor/Private/SRenameGameplayTagDialog.h
+++ b/Engine/Plugins/Editor/GameplayTagsEditor/Source/GameplayTagsEditor/Private/SRenameGameplayTagDialog.h
@@ -2,12 +2,11 @@
 
 #pragma once
 
-#include "CoreMinimal.h"
-#include "GameplayTagsManager.h"
-#include "Input/Reply.h"
 #include "Widgets/SCompoundWidget.h"
-#include "Widgets/DeclarativeSyntaxSupport.h"
-#include "Widgets/Input/SEditableTextBox.h"
+
+class SEditableTextBox;
+namespace ETextCommit { enum Type : int; }
+struct FGameplayTagNode;
 
 class SRenameGameplayTagDialog : public SCompoundWidget
 {
diff --git a/Engine/Plugins/Editor/GeometryMode/Source/BspMode/Private/BspModeModule.cpp b/Engine/Plugins/Editor/GeometryMode/Source/BspMode/Private/BspModeModule.cpp
--- a/Engine/Plugins/Editor/GeometryMode/Source/BspMode/Private/BspModeModule.cpp
+++ b/Engine/Plugins/Editor/GeometryMode/Source/BspMode/Private/BspModeModule.cpp
@@ -1,11 +1,7 @@
 // Copyright Epic Games, Inc. All Rights Reserved.
 
 #include "BspModeModule.h"
-#include "Modules/ModuleManager.h"
-#include "Widgets/DeclarativeSyntaxSupport.h"
-#include "Textures/SlateIcon.h"
 #include "Styling/AppStyle.h"
-#include "Engine/BrushBuilder.h"
 #include "Builders/ConeBuilder.h"
 #include "Builders/CubeBuilder.h"
 #include "Builders/CurvedStairBuilder.h"
@@ -13,11 +9,7 @@
 #include "Builders/LinearStairBuilder.h"
 #include "Builders/SpiralStairBuilder.h"
 #include "Builders/TetrahedronBuilder.h"
-#include "Widgets/SWidget.h"
-#include "EditorModeRegistry.h"
-#include "EditorModes.h"
 #include "BspModeStyle.h"
-#include "BspMode.h"
 #include "SBspPalette.h"
 #include "IPlacementModeModule.h"
 
diff --git a/Engine/Plugins/Editor/GeometryMode/Source/BspMode/Private/BspModeModule.h b/Engine/Plugins/Editor/GeometryMode/Source/BspMode/Private/BspModeModule.h
--- a/Engine/Plugins/Editor/GeometryMode/Source/BspMode/Private/BspModeModule.h
+++ b/Engine/Plugins/Editor/GeometryMode/Source/BspMode/Private/BspModeModule.h
@@ -2,10 +2,10 @@
 
 #pragma once
 
-#include "CoreMinimal.h"
 #include "UObject/Class.h"
 #include "UObject/WeakObjectPtr.h"
 #include "IBspModeModule.h"
+#include "UObject/WeakObjectPtrTemplates.h"
 
 struct FSlateBrush;
 
diff --git a/Engine/Plugins/Editor/GeometryMode/Source/BspMode/Private/BspModeStyle.cpp b/Engine/Plugins/Editor/GeometryMode/Source/BspMode/Private/BspModeStyle.cpp
--- a/Engine/Plugins/Editor/GeometryMode/Source/BspMode/Private/BspModeStyle.cpp
+++ b/Engine/Plugins/Editor/GeometryMode/Source/BspMode/Private/BspModeStyle.cpp
@@ -1,10 +1,15 @@
 // Copyright Epic Games, Inc. All Rights Reserved.
 
 #include "BspModeStyle.h"
+#include "Brushes/SlateBorderBrush.h"
 #include "Styling/SlateStyleRegistry.h"
+#include "Brushes/SlateBoxBrush.h"
 #include "Styling/SlateTypes.h"
+#include "Brushes/SlateImageBrush.h"
 #include "Styling/CoreStyle.h"
-#include "Styling/AppStyle.h"
+#include "Brushes/SlateNoResource.h"
+#include "Misc/Paths.h"
+#include "Styling/SlateStyle.h"
 
 
 TSharedPtr< FSlateStyleSet > FBspModeStyle::StyleSet = nullptr;
diff --git a/Engine/Plugins/Editor/GeometryMode/Source/BspMode/Private/SBspPalette.cpp b/Engine/Plugins/Editor/GeometryMode/Source/BspMode/Private/SBspPalette.cpp
--- a/Engine/Plugins/Editor/GeometryMode/Source/BspMode/Private/SBspPalette.cpp
+++ b/Engine/Plugins/Editor/GeometryMode/Source/BspMode/Private/SBspPalette.cpp
@@ -1,24 +1,14 @@
 // Copyright Epic Games, Inc. All Rights Reserved.
 
 #include "SBspPalette.h"
+#include "BspModeModule.h"
 #include "Modules/ModuleManager.h"
-#include "Widgets/SBoxPanel.h"
+#include "Framework/Views/TableViewMetadata.h"
 #include "SlateOptMacros.h"
-#include "Styling/CoreStyle.h"
-#include "Widgets/Images/SImage.h"
-#include "Styling/SlateTypes.h"
-#include "Widgets/Text/STextBlock.h"
 #include "Widgets/Input/SButton.h"
-#include "Widgets/Layout/SBox.h"
 #include "Widgets/Layout/SScrollBorder.h"
-#include "Widgets/Layout/SBorder.h"
-#include "Widgets/Views/STableViewBase.h"
-#include "Widgets/Views/STableRow.h"
+#include "Widgets/SToolTip.h"
 #include "Widgets/Views/SListView.h"
-#include "Widgets/Input/SCheckBox.h"
-#include "Styling/AppStyle.h"
-#include "Engine/Brush.h"
-#include "BspModeStyle.h"
 #include "DragAndDrop/BrushBuilderDragDropOp.h"
 #include "EditorClassUtils.h"
 
diff --git a/Engine/Plugins/Editor/GeometryMode/Source/BspMode/Private/SBspPalette.h b/Engine/Plugins/Editor/GeometryMode/Source/BspMode/Private/SBspPalette.h
--- a/Engine/Plugins/Editor/GeometryMode/Source/BspMode/Private/SBspPalette.h
+++ b/Engine/Plugins/Editor/GeometryMode/Source/BspMode/Private/SBspPalette.h
@@ -2,10 +2,6 @@
 
 #pragma once
 
-#include "CoreMinimal.h"
-#include "BspModeModule.h"
-#include "Widgets/DeclarativeSyntaxSupport.h"
-#include "Input/Reply.h"
 #include "Widgets/SCompoundWidget.h"
 #include "Misc/NotifyHook.h"
 
diff --git a/Engine/Plugins/Editor/GeometryMode/Source/GeometryMode/Private/BrushEditingSubsystemImpl.cpp b/Engine/Plugins/Editor/GeometryMode/Source/GeometryMode/Private/BrushEditingSubsystemImpl.cpp
--- a/Engine/Plugins/Editor/GeometryMode/Source/GeometryMode/Private/BrushEditingSubsystemImpl.cpp
+++ b/Engine/Plugins/Editor/GeometryMode/Source/GeometryMode/Private/BrushEditingSubsystemImpl.cpp
@@ -1,14 +1,13 @@
 // Copyright Epic Games, Inc. All Rights Reserved.
 
 #include "BrushEditingSubsystemImpl.h"
+#include "Editor/UnrealEdTypes.h"
 #include "EditorGeometry.h"
 #include "EditorModeManager.h"
-#include "EdMode.h"
-#include "EditorModeTools.h"
+#include "Engine/Brush.h"
 #include "GeometryEdMode.h"
 #include "LevelEditorViewport.h"
 #include "SnappingUtils.h"
-#include "HitProxies.h"
 #include "LevelViewportClickHandlers.h"
 #include "GeometryModeModule.h"
 #include "ScopedTransaction.h"
diff --git a/Engine/Plugins/Editor/GeometryMode/Source/GeometryMode/Private/EditorGeometry.cpp b/Engine/Plugins/Editor/GeometryMode/Source/GeometryMode/Private/EditorGeometry.cpp
--- a/Engine/Plugins/Editor/GeometryMode/Source/GeometryMode/Private/EditorGeometry.cpp
+++ b/Engine/Plugins/Editor/GeometryMode/Source/GeometryMode/Private/EditorGeometry.cpp
@@ -2,9 +2,9 @@
 
 
 #include "EditorGeometry.h"
+#include "Engine/Brush.h"
 #include "Engine/Polys.h"
 #include "EditorModeManager.h"
-#include "EditorModes.h"
 #include "GeometryEdMode.h"
 #include "GeometryModeModule.h"
 #include "Model.h"
diff --git a/Engine/Plugins/Editor/GeometryMode/Source/GeometryMode/Private/GeometryEdMode.cpp b/Engine/Plugins/Editor/GeometryMode/Source/GeometryMode/Private/GeometryEdMode.cpp
--- a/Engine/Plugins/Editor/GeometryMode/Source/GeometryMode/Private/GeometryEdMode.cpp
+++ b/Engine/Plugins/Editor/GeometryMode/Source/GeometryMode/Private/GeometryEdMode.cpp
@@ -3,13 +3,12 @@
 #include "GeometryEdMode.h"
 #include "EditorViewportClient.h"
 #include "Engine/Texture2D.h"
+#include "Math/RotationMatrix.h"
 #include "Misc/FeedbackContext.h"
-#include "Modules/ModuleManager.h"
 #include "SceneView.h"
-#include "Styling/AppStyle.h"
+#include "Selection.h"
 #include "Settings/EditorStyleSettings.h"
 #include "Materials/Material.h"
-#include "Engine/Selection.h"
 #include "EditorModeManager.h"
 #include "EditorModes.h"
 #include "TextureResource.h"
@@ -18,8 +17,6 @@
 #include "GeometryMode.h"
 #include "EditorGeometry.h"
 #include "DynamicMeshBuilder.h"
-#include "GeomModifier.h"
-#include "GeomModifier_Edit.h"
 #include "GeomModifier_Clip.h"
 #include "GeomModifier_Create.h"
 #include "GeomModifier_Delete.h"
@@ -28,14 +25,12 @@
 #include "GeomModifier_Lathe.h"
 #include "GeomModifier_Pen.h"
 #include "GeomModifier_Split.h"
-#include "GeomModifier_Triangulate.h"
 #include "GeomModifier_Optimize.h"
 #include "GeomModifier_Turn.h"
 #include "GeomModifier_Weld.h"
 #include "NavMesh/NavMeshBoundsVolume.h"
 #include "NavigationSystem.h"
 #include "GeometryModeModule.h"
-#include "Engine/Brush.h"
 
 DEFINE_LOG_CATEGORY_STATIC(LogGeometryMode, Log, All);
 
diff --git a/Engine/Plugins/Editor/GeometryMode/Source/GeometryMode/Private/GeometryMode.cpp b/Engine/Plugins/Editor/GeometryMode/Source/GeometryMode/Private/GeometryMode.cpp
--- a/Engine/Plugins/Editor/GeometryMode/Source/GeometryMode/Private/GeometryMode.cpp
+++ b/Engine/Plugins/Editor/GeometryMode/Source/GeometryMode/Private/GeometryMode.cpp
@@ -1,27 +1,23 @@
 // Copyright Epic Games, Inc. All Rights Reserved.
 
 #include "GeometryMode.h"
+#include "DetailsViewArgs.h"
 #include "Widgets/Text/STextBlock.h"
+#include "Misc/NotifyHook.h"
 #include "Modules/ModuleManager.h"
-#include "Widgets/Layout/SBorder.h"
 #include "Widgets/Layout/SSeparator.h"
 #include "Widgets/Layout/SGridPanel.h"
-#include "Widgets/Layout/SUniformGridPanel.h"
 #include "Widgets/Input/SButton.h"
 #include "Widgets/SToolTip.h"
 #include "Widgets/Layout/SScrollBox.h"
 #include "Widgets/Input/SCheckBox.h"
-#include "Styling/AppStyle.h"
 #include "EditorModeManager.h"
-#include "EditorModes.h"
 #include "GeometryEdMode.h"
 #include "PropertyEditorModule.h"
 #include "IDetailsView.h"
 #include "GeomModifier.h"
 #include "GeometryModeModule.h"
 #include "GeometryModeStyle.h"
-#include "Engine/Selection.h"
-#include "Settings/EditorStyleSettings.h"
 #include "Framework/MultiBox/MultiBoxBuilder.h"
 
 #define LOCTEXT_NAMESPACE "GeometryMode"
diff --git a/Engine/Plugins/Editor/GeometryMode/Source/GeometryMode/Private/GeometryMode.h b/Engine/Plugins/Editor/GeometryMode/Source/GeometryMode/Private/GeometryMode.h
--- a/Engine/Plugins/Editor/GeometryMode/Source/GeometryMode/Private/GeometryMode.h
+++ b/Engine/Plugins/Editor/GeometryMode/Source/GeometryMode/Private/GeometryMode.h
@@ -2,17 +2,10 @@
 
 #pragma once
 
-#include "CoreMinimal.h"
-#include "Layout/Visibility.h"
-#include "Input/Reply.h"
-#include "Widgets/DeclarativeSyntaxSupport.h"
-#include "Widgets/SWidget.h"
-#include "Widgets/SCompoundWidget.h"
-#include "Widgets/SBoxPanel.h"
-#include "Misc/NotifyHook.h"
-#include "Toolkits/IToolkitHost.h"
 #include "Toolkits/BaseToolkit.h"
 
+enum class ECheckBoxState : uint8;
+
 class IDetailsView;
 class SButton;
 class SCheckBox;
diff --git a/Engine/Plugins/Editor/GeometryMode/Source/GeometryMode/Private/GeometryModeModule.cpp b/Engine/Plugins/Editor/GeometryMode/Source/GeometryMode/Private/GeometryModeModule.cpp
--- a/Engine/Plugins/Editor/GeometryMode/Source/GeometryMode/Private/GeometryModeModule.cpp
+++ b/Engine/Plugins/Editor/GeometryMode/Source/GeometryMode/Private/GeometryModeModule.cpp
@@ -2,11 +2,10 @@
 
 #include "GeometryModeModule.h"
 #include "Modules/ModuleManager.h"
-#include "EditorModeManager.h"
-#include "Editor.h"
 #include "EditorModeRegistry.h"
 #include "GeometryEdMode.h"
 #include "GeometryModeStyle.h"
+#include "Styling/AppStyle.h"
 
 
 FEditorModeID FGeometryEditingModes::EM_Geometry = FEditorModeID(TEXT("EM_Geometry"));
diff --git a/Engine/Plugins/Editor/GeometryMode/Source/GeometryMode/Private/GeometryModeStyle.cpp b/Engine/Plugins/Editor/GeometryMode/Source/GeometryMode/Private/GeometryModeStyle.cpp
--- a/Engine/Plugins/Editor/GeometryMode/Source/GeometryMode/Private/GeometryModeStyle.cpp
+++ b/Engine/Plugins/Editor/GeometryMode/Source/GeometryMode/Private/GeometryModeStyle.cpp
@@ -1,12 +1,13 @@
 // Copyright Epic Games, Inc. All Rights Reserved.
 
 #include "GeometryModeStyle.h"
+#include "Brushes/SlateImageBrush.h"
 #include "Styling/SlateStyleRegistry.h"
+#include "Misc/Paths.h"
 #include "Styling/SlateTypes.h"
-#include "Styling/CoreStyle.h"
 #include "Styling/AppStyle.h"
 #include "Interfaces/IPluginManager.h"
-#include "SlateOptMacros.h"
+#include "Styling/SlateStyle.h"
 
 
 #define IMAGE_PLUGIN_BRUSH( RelativePath, ... ) FSlateImageBrush( FGeometryModeStyle::InContent( RelativePath, ".png" ), __VA_ARGS__ )
diff --git a/Engine/Plugins/Editor/GeometryMode/Source/GeometryMode/Private/GeometryModeStyle.h b/Engine/Plugins/Editor/GeometryMode/Source/GeometryMode/Private/GeometryModeStyle.h
--- a/Engine/Plugins/Editor/GeometryMode/Source/GeometryMode/Private/GeometryModeStyle.h
+++ b/Engine/Plugins/Editor/GeometryMode/Source/GeometryMode/Private/GeometryModeStyle.h
@@ -2,9 +2,7 @@
 
 #pragma once
 
-#include "CoreMinimal.h"
-#include "Styling/ISlateStyle.h"
-#include "Styling/SlateStyle.h"
+#include "Templates/SharedPointer.h"
 
 class FGeometryModeStyle
 {
diff --git a/Engine/Plugins/Editor/GeometryMode/Source/GeometryMode/Private/GeometryModifiers.cpp b/Engine/Plugins/Editor/GeometryMode/Source/GeometryMode/Private/GeometryModifiers.cpp
--- a/Engine/Plugins/Editor/GeometryMode/Source/GeometryMode/Private/GeometryModifiers.cpp
+++ b/Engine/Plugins/Editor/GeometryMode/Source/GeometryMode/Private/GeometryModifiers.cpp
@@ -1,21 +1,14 @@
 // Copyright Epic Games, Inc. All Rights Reserved.
 
-#include "CoreMinimal.h"
-#include "Editor.h"
-#include "EngineDefines.h"
+#include "Math/RotationMatrix.h"
 #include "Misc/MessageDialog.h"
-#include "InputCoreTypes.h"
-#include "GameFramework/Actor.h"
-#include "Engine/World.h"
 #include "SceneView.h"
 #include "Model.h"
+#include "Selection.h"
 #include "Settings/LevelEditorViewportSettings.h"
-#include "Engine/Brush.h"
 #include "Settings/LevelEditorMiscSettings.h"
 #include "Engine/Polys.h"
-#include "Engine/Selection.h"
 #include "EditorModeManager.h"
-#include "EditorModes.h"
 #include "Dialogs/Dialogs.h"
 #include "BSPOps.h"
 #include "Engine/BrushBuilder.h"
@@ -30,8 +23,6 @@
 #include "Framework/Notifications/NotificationManager.h"
 #include "Widgets/Notifications/SNotificationList.h"
 #include "Components/BrushComponent.h"
-#include "GeomModifier.h"
-#include "GeomModifier_Edit.h"
 #include "GeomModifier_Clip.h"
 #include "GeomModifier_Create.h"
 #include "GeomModifier_Delete.h"
@@ -40,7 +31,6 @@
 #include "GeomModifier_Lathe.h"
 #include "GeomModifier_Pen.h"
 #include "GeomModifier_Split.h"
-#include "GeomModifier_Triangulate.h"
 #include "GeomModifier_Optimize.h"
 #include "GeomModifier_Turn.h"
 #include "GeomModifier_Weld.h"
diff --git a/Engine/Plugins/Editor/GeometryMode/Source/TextureAlignMode/Private/TextureAlignEdMode.cpp b/Engine/Plugins/Editor/GeometryMode/Source/TextureAlignMode/Private/TextureAlignEdMode.cpp
--- a/Engine/Plugins/Editor/GeometryMode/Source/TextureAlignMode/Private/TextureAlignEdMode.cpp
+++ b/Engine/Plugins/Editor/GeometryMode/Source/TextureAlignMode/Private/TextureAlignEdMode.cpp
@@ -1,10 +1,12 @@
 // Copyright Epic Games, Inc. All Rights Reserved.
 
 #include "TextureAlignEdMode.h"
+#include "EditorModeRegistry.h"
 #include "EditorViewportClient.h"
+#include "Engine/Brush.h"
 #include "Modules/ModuleManager.h"
 #include "EditorModeManager.h"
-#include "EditorModes.h"
+#include "Math/RotationMatrix.h"
 #include "ScopedTransaction.h"
 #include "SurfaceIterators.h"
 #include "EditorSupportDelegates.h"
diff --git a/Engine/Plugins/Editor/GeometryMode/Source/TextureAlignMode/Private/TextureAlignMode.cpp b/Engine/Plugins/Editor/GeometryMode/Source/TextureAlignMode/Private/TextureAlignMode.cpp
--- a/Engine/Plugins/Editor/GeometryMode/Source/TextureAlignMode/Private/TextureAlignMode.cpp
+++ b/Engine/Plugins/Editor/GeometryMode/Source/TextureAlignMode/Private/TextureAlignMode.cpp
@@ -2,7 +2,6 @@
 
 #include "TextureAlignMode.h"
 #include "EditorModeManager.h"
-#include "EditorModes.h"
 #include "TextureAlignEdMode.h"
 #include "GeometryModeModule.h"
 
diff --git a/Engine/Plugins/Editor/MaterialAnalyzer/Source/MaterialAnalyzerModule.cpp b/Engine/Plugins/Editor/MaterialAnalyzer/Source/MaterialAnalyzerModule.cpp
--- a/Engine/Plugins/Editor/MaterialAnalyzer/Source/MaterialAnalyzerModule.cpp
+++ b/Engine/Plugins/Editor/MaterialAnalyzer/Source/MaterialAnalyzerModule.cpp
@@ -2,14 +2,8 @@
 
 #include "MaterialAnalyzerModule.h"
 #include "SMaterialAnalyzer.h"
-#include "EditorModeRegistry.h"
-#include "Modules/ModuleManager.h"
-#include "Framework/Docking/TabManager.h"
-#include "ISettingsModule.h"
-#include "Editor.h"
 #include "WorkspaceMenuStructure.h"
 #include "WorkspaceMenuStructureModule.h"
-#include "Styling/AppStyle.h"
 #include "Widgets/Docking/SDockTab.h"
 
 #define LOCTEXT_NAMESPACE "MaterialAnalyzer"
@@ -65,4 +59,4 @@ TSharedRef<SDockTab> FMaterialAnalyzerModule::SpawnMaterialAnalyzerTab(const FSp
 
 IMPLEMENT_MODULE(FMaterialAnalyzerModule, MaterialAnalyzer)
 
-#undef LOCTEXT_NAMESPACE
\ No newline at end of file
+#undef LOCTEXT_NAMESPACE
diff --git a/Engine/Plugins/Editor/MaterialAnalyzer/Source/Private/AnalyzedMaterialNode.h b/Engine/Plugins/Editor/MaterialAnalyzer/Source/Private/AnalyzedMaterialNode.h
--- a/Engine/Plugins/Editor/MaterialAnalyzer/Source/Private/AnalyzedMaterialNode.h
+++ b/Engine/Plugins/Editor/MaterialAnalyzer/Source/Private/AnalyzedMaterialNode.h
@@ -2,7 +2,6 @@
 
 #pragma once
 
-#include "CoreMinimal.h"
 #include "AssetRegistry/AssetData.h"
 
 struct FBasePropertyOverrideNode
@@ -171,4 +170,4 @@ public:
 
 protected:
 	TArray<FAnalyzedMaterialNodeRef> ChildNodes;
-};
\ No newline at end of file
+};
diff --git a/Engine/Plugins/Editor/MaterialAnalyzer/Source/Private/SAnalyzedMaterialNodeWidgetItem.cpp b/Engine/Plugins/Editor/MaterialAnalyzer/Source/Private/SAnalyzedMaterialNodeWidgetItem.cpp
--- a/Engine/Plugins/Editor/MaterialAnalyzer/Source/Private/SAnalyzedMaterialNodeWidgetItem.cpp
+++ b/Engine/Plugins/Editor/MaterialAnalyzer/Source/Private/SAnalyzedMaterialNodeWidgetItem.cpp
@@ -1,14 +1,11 @@
 // Copyright Epic Games, Inc. All Rights Reserved.
 
 #include "SAnalyzedMaterialNodeWidgetItem.h"
-#include "Widgets/Views/STreeView.h"
-#include "SlateOptMacros.h"
-#include "Widgets/Images/SImage.h"
-#include "Engine/EngineTypes.h"
 #include "MaterialShaderType.h"
-#include "PropertyCustomizationHelpers.h"
 #include "Editor.h"
-#include "Styling/AppStyle.h"
+#include "Textures/SlateIcon.h"
+#include "Widgets/Input/SButton.h"
+#include "Widgets/Views/SListView.h"
 
 #define LOCTEXT_NAMESPACE "MaterialAnalyzer"
 
@@ -420,4 +417,4 @@ const FSlateBrush* SStaticParameterWidget<NodeType>::GetExpanderImage() const
 }
 
 
-#undef LOCTEXT_NAMESPACE
\ No newline at end of file
+#undef LOCTEXT_NAMESPACE
diff --git a/Engine/Plugins/Editor/MaterialAnalyzer/Source/Private/SAnalyzedMaterialNodeWidgetItem.h b/Engine/Plugins/Editor/MaterialAnalyzer/Source/Private/SAnalyzedMaterialNodeWidgetItem.h
--- a/Engine/Plugins/Editor/MaterialAnalyzer/Source/Private/SAnalyzedMaterialNodeWidgetItem.h
+++ b/Engine/Plugins/Editor/MaterialAnalyzer/Source/Private/SAnalyzedMaterialNodeWidgetItem.h
@@ -2,16 +2,12 @@
 
 #pragma once
 
-#include "CoreMinimal.h"
-#include "Styling/SlateColor.h"
-#include "Widgets/DeclarativeSyntaxSupport.h"
-#include "Widgets/SWidget.h"
-#include "Widgets/Views/STableViewBase.h"
+#include "Widgets/Views/ITableRow.h"
 #include "Widgets/Views/STableRow.h"
-#include "Widgets/Views/SListView.h"
 #include "AnalyzedMaterialNode.h"
-#include "Widgets/Input/SButton.h"
-#include "AssetRegistry/AssetData.h"
+
+class STableViewBase;
+class SVerticalBox;
 
 class SAnalyzedMaterialNodeWidgetItem
 	: public SMultiColumnTableRow<FAnalyzedMaterialNodeRef>
@@ -145,4 +141,4 @@ class SStaticMaterialLayerParameterWidget
 public:
 	virtual TSharedRef<SWidget> CreateRowWidget(FStaticMaterialLayerParameterNodeRef RowData) override;
 	virtual FText GetBaseText() const override;
-};
\ No newline at end of file
+};
diff --git a/Engine/Plugins/Editor/MaterialAnalyzer/Source/Private/SMaterialAnalyzer.cpp b/Engine/Plugins/Editor/MaterialAnalyzer/Source/Private/SMaterialAnalyzer.cpp
--- a/Engine/Plugins/Editor/MaterialAnalyzer/Source/Private/SMaterialAnalyzer.cpp
+++ b/Engine/Plugins/Editor/MaterialAnalyzer/Source/Private/SMaterialAnalyzer.cpp
@@ -1,21 +1,22 @@
 // Copyright Epic Games, Inc. All Rights Reserved.
 
 #include "SMaterialAnalyzer.h"
-#include "Widgets/Input/SSearchBox.h"
 #include "AssetRegistry/AssetRegistryModule.h"
-#include "Async/AsyncWork.h"
-#include "Framework/MultiBox/MultiBoxBuilder.h"
-#include "Styling/SlateIconFinder.h"
-#include "Framework/Commands/UIAction.h"
-#include "Materials/MaterialLayersFunctions.h"
-#include "Hash/CityHash.h"
-#include "Styling/AppStyle.h"
+#include "ContentBrowserDataLegacyBridge.h"
 #include "PropertyCustomizationHelpers.h"
-#include "CollectionManagerTypes.h"
 #include "CollectionManagerModule.h"
+#include "Framework/Views/TableViewMetadata.h"
 #include "ICollectionManager.h"
 #include "AssetManagerEditorModule.h"
-#include "Widgets/Images/SImage.h"
+#include "Materials/Material.h"
+#include "Materials/MaterialInstance.h"
+#include "Misc/Paths.h"
+#include "SAnalyzedMaterialNodeWidgetItem.h"
+#include "Widgets/Images/SThrobber.h"
+#include "Widgets/Input/SButton.h"
+#include "Widgets/Input/SEditableText.h"
+#include "Widgets/Layout/SScrollBox.h"
+#include "Widgets/Views/STreeView.h"
 
 #define LOCTEXT_NAMESPACE "MaterialAnalyzer"
 
@@ -991,4 +992,4 @@ void FAnalyzeForIdenticalPermutationsAsyncTask::GatherSuggestions()
 }
 
 
-#undef LOCTEXT_NAMESPACE
\ No newline at end of file
+#undef LOCTEXT_NAMESPACE
diff --git a/Engine/Plugins/Editor/MaterialAnalyzer/Source/Private/SMaterialAnalyzer.h b/Engine/Plugins/Editor/MaterialAnalyzer/Source/Private/SMaterialAnalyzer.h
--- a/Engine/Plugins/Editor/MaterialAnalyzer/Source/Private/SMaterialAnalyzer.h
+++ b/Engine/Plugins/Editor/MaterialAnalyzer/Source/Private/SMaterialAnalyzer.h
@@ -3,23 +3,18 @@
 #pragma once
 
 #include "Async/AsyncWork.h"
-#include "CoreMinimal.h"
-#include "Containers/Queue.h"
-#include "Input/Reply.h"
-#include "Widgets/DeclarativeSyntaxSupport.h"
-#include "Engine/GameViewportClient.h"
+#include "MaterialTypes.h"
 #include "Widgets/SCompoundWidget.h"
-#include "Framework/Commands/UICommandList.h"
-#include "Framework/Docking/TabManager.h"
-#include "Widgets/Views/STreeView.h"
 #include "AnalyzedMaterialNode.h"
-#include "SAnalyzedMaterialNodeWidgetItem.h"
-#include "Materials/Material.h"
-#include "Materials/MaterialInstance.h"
-#include "ContentBrowserModule.h"
-#include "Widgets/Input/SComboButton.h"
-#include "Widgets/Layout/SScrollBox.h"
-#include "Widgets/Images/SThrobber.h"
+
+class ITableRow;
+class SDockTab;
+class SScrollBox;
+class STableViewBase;
+class STextBlock;
+class SThrobber;
+class UMaterialInterface;
+template <typename ItemType> class STreeView;
 
 struct FBuildBasicMaterialTreeAsyncTask : FNonAbandonableTask
 {
@@ -286,4 +281,4 @@ private:
 
 private:
 	bool bAllowMaterialSelection;
-};
\ No newline at end of file
+};
diff --git a/Engine/Plugins/Editor/ModelingToolsEditorMode/Source/ModelingToolsEditorMode/Private/DetailsCustomizations/BakeMeshAttributeToolCustomizations.cpp b/Engine/Plugins/Editor/ModelingToolsEditorMode/Source/ModelingToolsEditorMode/Private/DetailsCustomizations/BakeMeshAttributeToolCustomizations.cpp
--- a/Engine/Plugins/Editor/ModelingToolsEditorMode/Source/ModelingToolsEditorMode/Private/DetailsCustomizations/BakeMeshAttributeToolCustomizations.cpp
+++ b/Engine/Plugins/Editor/ModelingToolsEditorMode/Source/ModelingToolsEditorMode/Private/DetailsCustomizations/BakeMeshAttributeToolCustomizations.cpp
@@ -1,21 +1,13 @@
 // Copyright Epic Games, Inc. All Rights Reserved.
 
 #include "DetailsCustomizations/BakeMeshAttributeToolCustomizations.h"
-#include "DetailCategoryBuilder.h"
 #include "DetailLayoutBuilder.h"
 #include "DetailWidgetRow.h"
-#include "PropertyCustomizationHelpers.h"
 #include "Widgets/Input/SComboButton.h"
-#include "IDetailChildrenBuilder.h"
-#include "Internationalization/BreakIterator.h"
 #include "Framework/MultiBox/MultiBoxBuilder.h"
-#include "IContentBrowserSingleton.h"
-#include "SAssetView.h"
 
-#include "ModelingWidgets/SDynamicNumericEntry.h"
 #include "ModelingWidgets/ModelingCustomizationUtil.h"
 
-#include "ModelingToolsEditorModeSettings.h"
 
 #include "BakeMeshAttributeMapsTool.h"
 #include "BakeMultiMeshAttributeMapsTool.h"
diff --git a/Engine/Plugins/Editor/ModelingToolsEditorMode/Source/ModelingToolsEditorMode/Private/DetailsCustomizations/BakeTransformToolCustomizations.cpp b/Engine/Plugins/Editor/ModelingToolsEditorMode/Source/ModelingToolsEditorMode/Private/DetailsCustomizations/BakeTransformToolCustomizations.cpp
--- a/Engine/Plugins/Editor/ModelingToolsEditorMode/Source/ModelingToolsEditorMode/Private/DetailsCustomizations/BakeTransformToolCustomizations.cpp
+++ b/Engine/Plugins/Editor/ModelingToolsEditorMode/Source/ModelingToolsEditorMode/Private/DetailsCustomizations/BakeTransformToolCustomizations.cpp
@@ -1,7 +1,6 @@
 // Copyright Epic Games, Inc. All Rights Reserved.
 
 #include "DetailsCustomizations/BakeTransformToolCustomizations.h"
-#include "UObject/Class.h"
 #include "DetailLayoutBuilder.h"
 #include "PropertyHandle.h"
 #include "PropertyRestriction.h"
diff --git a/Engine/Plugins/Editor/ModelingToolsEditorMode/Source/ModelingToolsEditorMode/Private/DetailsCustomizations/MeshVertexSculptToolCustomizations.cpp b/Engine/Plugins/Editor/ModelingToolsEditorMode/Source/ModelingToolsEditorMode/Private/DetailsCustomizations/MeshVertexSculptToolCustomizations.cpp
--- a/Engine/Plugins/Editor/ModelingToolsEditorMode/Source/ModelingToolsEditorMode/Private/DetailsCustomizations/MeshVertexSculptToolCustomizations.cpp
+++ b/Engine/Plugins/Editor/ModelingToolsEditorMode/Source/ModelingToolsEditorMode/Private/DetailsCustomizations/MeshVertexSculptToolCustomizations.cpp
@@ -1,36 +1,22 @@
 // Copyright Epic Games, Inc. All Rights Reserved.
 
 #include "DetailsCustomizations/MeshVertexSculptToolCustomizations.h"
-#include "DetailCategoryBuilder.h"
 #include "DetailLayoutBuilder.h"
 #include "DetailWidgetRow.h"
-#include "PropertyCustomizationHelpers.h"
-#include "Widgets/Input/SButton.h"
+#include "Engine/Texture2D.h"
 #include "Widgets/Input/SCheckBox.h"
-#include "Widgets/Text/STextBlock.h"
-#include "Widgets/Input/SComboButton.h"
-#include "Widgets/Views/STileView.h"
-#include "Widgets/Images/SImage.h"
-#include "IDetailChildrenBuilder.h"
-#include "Internationalization/BreakIterator.h"
-#include "Framework/MultiBox/MultiBoxBuilder.h"
-#include "IContentBrowserSingleton.h"
-#include "ContentBrowserModule.h"
-#include "SAssetView.h"
-#include "SAssetDropTarget.h"
 
 #include "ModelingWidgets/SComboPanel.h"
-#include "ModelingWidgets/SDynamicNumericEntry.h"
-#include "ModelingWidgets/SToolInputAssetPicker.h"
 #include "ModelingWidgets/SToolInputAssetComboPanel.h"
 #include "ModelingWidgets/ModelingCustomizationUtil.h"
 #include "DetailsCustomizations/ModelingToolsBrushSizeCustomization.h"
 
-#include "Styling/SlateStyleMacros.h"
 #include "ModelingToolsEditorModeStyle.h"
 #include "ModelingToolsEditorModeSettings.h"
 
 #include "MeshVertexSculptTool.h"
+#include "Widgets/Layout/SBox.h"
+#include "Widgets/SBoxPanel.h"
 
 using namespace UE::ModelingUI;
 
@@ -528,4 +514,4 @@ void FVertexBrushAlphaPropertiesDetails::CustomizeDetails(IDetailLayoutBuilder&
 
 
 
-#undef LOCTEXT_NAMESPACE
\ No newline at end of file
+#undef LOCTEXT_NAMESPACE
diff --git a/Engine/Plugins/Editor/ModelingToolsEditorMode/Source/ModelingToolsEditorMode/Private/DetailsCustomizations/ModelingToolsAxisFilterCustomization.cpp b/Engine/Plugins/Editor/ModelingToolsEditorMode/Source/ModelingToolsEditorMode/Private/DetailsCustomizations/ModelingToolsAxisFilterCustomization.cpp
--- a/Engine/Plugins/Editor/ModelingToolsEditorMode/Source/ModelingToolsEditorMode/Private/DetailsCustomizations/ModelingToolsAxisFilterCustomization.cpp
+++ b/Engine/Plugins/Editor/ModelingToolsEditorMode/Source/ModelingToolsEditorMode/Private/DetailsCustomizations/ModelingToolsAxisFilterCustomization.cpp
@@ -2,13 +2,11 @@
 
 #include "DetailsCustomizations/ModelingToolPropertyCustomizations.h"
 
-#include "PropertyHandle.h"
-#include "IPropertyTypeCustomization.h"
-#include "IPropertyUtilities.h"
 #include "DetailWidgetRow.h"
 
 
 #include "PropertySets/AxisFilterPropertyType.h"
+#include "Widgets/SBoxPanel.h"
 
 
 #define LOCTEXT_NAMESPACE "ModelingToolsAxisFilterCustomization"
@@ -98,4 +96,4 @@ void FModelingToolsAxisFilterCustomization::CustomizeChildren(TSharedRef<IProper
 }
 
 
-#undef LOCTEXT_NAMESPACE
\ No newline at end of file
+#undef LOCTEXT_NAMESPACE
diff --git a/Engine/Plugins/Editor/ModelingToolsEditorMode/Source/ModelingToolsEditorMode/Private/DetailsCustomizations/ModelingToolsBrushSizeCustomization.cpp b/Engine/Plugins/Editor/ModelingToolsEditorMode/Source/ModelingToolsEditorMode/Private/DetailsCustomizations/ModelingToolsBrushSizeCustomization.cpp
--- a/Engine/Plugins/Editor/ModelingToolsEditorMode/Source/ModelingToolsEditorMode/Private/DetailsCustomizations/ModelingToolsBrushSizeCustomization.cpp
+++ b/Engine/Plugins/Editor/ModelingToolsEditorMode/Source/ModelingToolsEditorMode/Private/DetailsCustomizations/ModelingToolsBrushSizeCustomization.cpp
@@ -2,16 +2,13 @@
 
 #include "DetailsCustomizations/ModelingToolsBrushSizeCustomization.h"
 
-#include "PropertyHandle.h"
-#include "IPropertyTypeCustomization.h"
-#include "IPropertyUtilities.h"
 #include "DetailWidgetRow.h"
 #include "Widgets/Input/SCheckBox.h"
 #include "Widgets/Layout/SBox.h"
-#include "ModelingWidgets/SDynamicNumericEntry.h"
 #include "ModelingWidgets/ModelingCustomizationUtil.h"
 
 #include "Sculpting/MeshSculptToolBase.h"
+#include "Widgets/SBoxPanel.h"
 
 using namespace UE::ModelingUI;
 
@@ -139,4 +136,4 @@ void FModelingToolsBrushSizeCustomization::CustomizeChildren(TSharedRef<IPropert
 }
 
 
-#undef LOCTEXT_NAMESPACE
\ No newline at end of file
+#undef LOCTEXT_NAMESPACE
diff --git a/Engine/Plugins/Editor/ModelingToolsEditorMode/Source/ModelingToolsEditorMode/Private/ModelingModeAssetUtils.cpp b/Engine/Plugins/Editor/ModelingToolsEditorMode/Source/ModelingToolsEditorMode/Private/ModelingModeAssetUtils.cpp
--- a/Engine/Plugins/Editor/ModelingToolsEditorMode/Source/ModelingToolsEditorMode/Private/ModelingModeAssetUtils.cpp
+++ b/Engine/Plugins/Editor/ModelingToolsEditorMode/Source/ModelingToolsEditorMode/Private/ModelingModeAssetUtils.cpp
@@ -1,13 +1,15 @@
 // Copyright Epic Games, Inc. All Rights Reserved.
 
 #include "ModelingModeAssetUtils.h"
+#include "ContentBrowserItemPath.h"
 #include "ModelingToolsEditorModeSettings.h"
 
-#include "Engine/World.h"
 #include "AssetRegistry/AssetRegistryModule.h"
 #include "AssetToolsModule.h"
 #include "FileHelpers.h"
+#include "IAssetTools.h"
 #include "Misc/Paths.h"
+#include "Misc/PackageName.h"
 #include "ModelingObjectsCreationAPI.h"
 #include "Editor/EditorEngine.h"
 #include "EditorDirectories.h"
@@ -15,6 +17,7 @@
 // for content-browser things
 #include "ContentBrowserModule.h"
 #include "IContentBrowserSingleton.h"
+#include "Trace/Trace.inl"
 
 extern UNREALED_API UEditorEngine* GEditor;
 
@@ -269,4 +272,4 @@ void UE::Modeling::OnNewAssetCreated(UObject* Asset)
 }
 
 
-#undef LOCTEXT_NAMESPACE
\ No newline at end of file
+#undef LOCTEXT_NAMESPACE
diff --git a/Engine/Plugins/Editor/ModelingToolsEditorMode/Source/ModelingToolsEditorMode/Private/ModelingModeToolkit_Toolbars.cpp b/Engine/Plugins/Editor/ModelingToolsEditorMode/Source/ModelingToolsEditorMode/Private/ModelingModeToolkit_Toolbars.cpp
--- a/Engine/Plugins/Editor/ModelingToolsEditorMode/Source/ModelingToolsEditorMode/Private/ModelingModeToolkit_Toolbars.cpp
+++ b/Engine/Plugins/Editor/ModelingToolsEditorMode/Source/ModelingToolsEditorMode/Private/ModelingModeToolkit_Toolbars.cpp
@@ -1,12 +1,13 @@
 // Copyright Epic Games, Inc. All Rights Reserved.
 
 #include "ModelingToolsEditorModeToolkit.h"
+#include "EdModeInteractiveToolsContext.h"
 #include "ModelingToolsEditorMode.h"
+#include "Framework/Commands/UICommandList.h"
 #include "ModelingToolsManagerActions.h"
 #include "ModelingToolsEditorModeSettings.h"
 #include "ModelingToolsEditorModeStyle.h"
 
-#include "Selection/GeometrySelectionManager.h"
 #include "ModelingSelectionInteraction.h"
 
 #include "ModelingComponentsSettings.h"
@@ -16,14 +17,12 @@
 #include "ISettingsModule.h"
 
 #include "Framework/MultiBox/MultiBoxBuilder.h"
-#include "Framework/Commands/UIAction.h"
 #include "Widgets/Images/SImage.h"
+#include "Widgets/Layout/SBorder.h"
 #include "Widgets/Layout/SSpacer.h"
 #include "Widgets/Input/SComboButton.h"
 
 // to toggle component instance selection
-#include "ComponentRecreateRenderStateContext.h"
-#include "Elements/SMInstance/SMInstanceElementData.h"
 
 #define LOCTEXT_NAMESPACE "FModelingToolsEditorModeToolkit"
 
@@ -469,4 +468,4 @@ void FModelingToolsEditorModeToolkit::MakeSelectionPaletteOverlayWidget()
 }
 
 
-#undef LOCTEXT_NAMESPACE
\ No newline at end of file
+#undef LOCTEXT_NAMESPACE
diff --git a/Engine/Plugins/Editor/ModelingToolsEditorMode/Source/ModelingToolsEditorMode/Private/ModelingSelectionInteraction.cpp b/Engine/Plugins/Editor/ModelingToolsEditorMode/Source/ModelingToolsEditorMode/Private/ModelingSelectionInteraction.cpp
--- a/Engine/Plugins/Editor/ModelingToolsEditorMode/Source/ModelingToolsEditorMode/Private/ModelingSelectionInteraction.cpp
+++ b/Engine/Plugins/Editor/ModelingToolsEditorMode/Source/ModelingToolsEditorMode/Private/ModelingSelectionInteraction.cpp
@@ -2,13 +2,16 @@
 
 
 #include "ModelingSelectionInteraction.h"
+#include "BaseGizmos/CombinedTransformGizmo.h"
 #include "InteractiveToolsContext.h"
 #include "BaseGizmos/TransformGizmoUtil.h"
+#include "BaseGizmos/TransformProxy.h"
 #include "Selection/GeometrySelectionManager.h"
+#include "InteractiveGizmoManager.h"
 #include "ToolSceneQueriesUtil.h"
+#include "InteractiveToolManager.h"
 #include "SceneQueries/SceneSnappingManager.h"
 
-#include "BaseBehaviors/SingleClickBehavior.h"
 #include "BaseBehaviors/SingleClickOrDragBehavior.h"
 #include "BaseBehaviors/MouseHoverBehavior.h"
 #include "BaseBehaviors/KeyAsModifierInputBehavior.h"
@@ -518,4 +521,4 @@ void UPathSelectionInteraction::OnTerminateDragSequence()
 	{
 		SelectionManager->EndTrackedSelectionChange();
 	}
-}
\ No newline at end of file
+}
diff --git a/Engine/Plugins/Editor/ModelingToolsEditorMode/Source/ModelingToolsEditorMode/Private/ModelingToolsActions.cpp b/Engine/Plugins/Editor/ModelingToolsEditorMode/Source/ModelingToolsEditorMode/Private/ModelingToolsActions.cpp
--- a/Engine/Plugins/Editor/ModelingToolsEditorMode/Source/ModelingToolsEditorMode/Private/ModelingToolsActions.cpp
+++ b/Engine/Plugins/Editor/ModelingToolsEditorMode/Source/ModelingToolsEditorMode/Private/ModelingToolsActions.cpp
@@ -13,7 +13,7 @@
 #include "EditMeshPolygonsTool.h"
 #include "ShapeSprayTool.h"
 #include "MeshSpaceDeformerTool.h"
-#include "MeshSelectionTool.h"
+#include "Tools/StandardToolModeCommands.h"
 #include "TransformMeshesTool.h"
 #include "PlaneCutTool.h"
 #include "EditMeshPolygonsTool.h"
diff --git a/Engine/Plugins/Editor/ModelingToolsEditorMode/Source/ModelingToolsEditorMode/Private/ModelingToolsEditorMode.cpp b/Engine/Plugins/Editor/ModelingToolsEditorMode/Source/ModelingToolsEditorMode/Private/ModelingToolsEditorMode.cpp
--- a/Engine/Plugins/Editor/ModelingToolsEditorMode/Source/ModelingToolsEditorMode/Private/ModelingToolsEditorMode.cpp
+++ b/Engine/Plugins/Editor/ModelingToolsEditorMode/Source/ModelingToolsEditorMode/Private/ModelingToolsEditorMode.cpp
@@ -1,22 +1,23 @@
 // Copyright Epic Games, Inc. All Rights Reserved.
 
 #include "ModelingToolsEditorMode.h"
+#include "Components/StaticMeshComponent.h"
 #include "IGeometryProcessingInterfacesModule.h"
+#include "EdModeInteractiveToolsContext.h"
 #include "GeometryProcessingInterfaces/IUVEditorModularFeature.h"
-#include "InteractiveTool.h"
+#include "IAnalyticsProviderET.h"
 #include "ModelingToolsEditorModeToolkit.h"
+#include "ILevelEditor.h"
 #include "ModelingToolsEditorModeSettings.h"
-#include "Toolkits/ToolkitManager.h"
+#include "IStylusState.h"
 #include "ToolTargetManager.h"
 #include "ContextObjectStore.h"
+#include "InputRouter.h"
 #include "ToolTargets/StaticMeshComponentToolTarget.h"
+#include "Selection.h"
 #include "ToolTargets/VolumeComponentToolTarget.h"
 #include "ToolTargets/DynamicMeshComponentToolTarget.h"
 #include "ToolTargets/SkeletalMeshComponentToolTarget.h"
-#include "Framework/Commands/UICommandList.h"
-#include "Framework/Application/SlateApplication.h"
-#include "Framework/MultiBox/MultiBoxBuilder.h"
-#include "EditorViewportClient.h"
 #include "EngineAnalytics.h"
 #include "BaseGizmos/TransformGizmoUtil.h"
 #include "Selection/StoredMeshSelectionUtil.h"
@@ -38,7 +39,6 @@
 #include "RevolveBoundaryTool.h"
 #include "SmoothMeshTool.h"
 #include "OffsetMeshTool.h"
-#include "RemeshMeshTool.h"
 #include "SimplifyMeshTool.h"
 #include "MeshInspectorTool.h"
 #include "WeldMeshEdgesTool.h"
@@ -46,8 +46,6 @@
 #include "DrawPolyPathTool.h"
 #include "DrawAndRevolveTool.h"
 #include "ShapeSprayTool.h"
-#include "MergeMeshesTool.h"
-#include "VoxelCSGMeshesTool.h"
 #include "VoxelSolidifyMeshesTool.h"
 #include "VoxelBlendMeshesTool.h"
 #include "VoxelMorphologyMeshesTool.h"
@@ -67,7 +65,6 @@
 #include "RemoveOccludedTrianglesTool.h"
 #include "AttributeEditorTool.h"
 #include "TransformMeshesTool.h"
-#include "MeshSelectionTool.h"
 #include "UVProjectionTool.h"
 #include "UVLayoutTool.h"
 #include "EditMeshMaterialsTool.h"
@@ -117,7 +114,6 @@
 #include "IStylusInputModule.h"
 
 #include "LevelEditor.h"
-#include "IAssetViewport.h"
 #include "SLevelViewport.h"
 #include "Application/ThrottleManager.h"
 
@@ -126,19 +122,21 @@
 #include "ModelingModeAssetUtils.h"
 #include "EditorModelingObjectsCreationAPI.h"
 
-#include "Selections/GeometrySelection.h"
 #include "Selection/GeometrySelectionManager.h"
-#include "Selection/DynamicMeshSelector.h"
 #include "Selection/VolumeSelector.h"
 #include "Selection/StaticMeshSelector.h"
 #include "ModelingSelectionInteraction.h"
 #include "DynamicMeshActor.h"
-#include "GameFramework/Volume.h"
 #include "Components/BrushComponent.h"
 #include "Engine/StaticMeshActor.h"
 
 #include UE_INLINE_GENERATED_CPP_BY_NAME(ModelingToolsEditorMode)
 
+#if WITH_PROXYLOD
+#include "MergeMeshesTool.h"
+#include "VoxelCSGMeshesTool.h"
+#endif
+
 #define LOCTEXT_NAMESPACE "UModelingToolsEditorMode"
 
 
diff --git a/Engine/Plugins/Editor/ModelingToolsEditorMode/Source/ModelingToolsEditorMode/Private/ModelingToolsEditorModeModule.cpp b/Engine/Plugins/Editor/ModelingToolsEditorMode/Source/ModelingToolsEditorMode/Private/ModelingToolsEditorModeModule.cpp
--- a/Engine/Plugins/Editor/ModelingToolsEditorMode/Source/ModelingToolsEditorMode/Private/ModelingToolsEditorModeModule.cpp
+++ b/Engine/Plugins/Editor/ModelingToolsEditorMode/Source/ModelingToolsEditorMode/Private/ModelingToolsEditorModeModule.cpp
@@ -5,7 +5,6 @@
 #include "ModelingToolsActions.h"
 #include "ModelingToolsManagerActions.h"
 #include "ModelingToolsEditorModeStyle.h"
-#include "ModelingToolsEditorModeSettings.h"
 
 #include "Misc/CoreDelegates.h"
 
@@ -105,4 +104,4 @@ void FModelingToolsEditorModeModule::OnPostEngineInit()
 
 #undef LOCTEXT_NAMESPACE
 	
-IMPLEMENT_MODULE(FModelingToolsEditorModeModule, ModelingToolsEditorMode)
\ No newline at end of file
+IMPLEMENT_MODULE(FModelingToolsEditorModeModule, ModelingToolsEditorMode)
diff --git a/Engine/Plugins/Editor/ModelingToolsEditorMode/Source/ModelingToolsEditorMode/Private/ModelingToolsEditorModeStyle.cpp b/Engine/Plugins/Editor/ModelingToolsEditorMode/Source/ModelingToolsEditorMode/Private/ModelingToolsEditorModeStyle.cpp
--- a/Engine/Plugins/Editor/ModelingToolsEditorMode/Source/ModelingToolsEditorMode/Private/ModelingToolsEditorModeStyle.cpp
+++ b/Engine/Plugins/Editor/ModelingToolsEditorMode/Source/ModelingToolsEditorMode/Private/ModelingToolsEditorModeStyle.cpp
@@ -1,12 +1,13 @@
 // Copyright Epic Games, Inc. All Rights Reserved.
 
 #include "ModelingToolsEditorModeStyle.h"
+#include "Brushes/SlateImageBrush.h"
 #include "Styling/SlateStyleRegistry.h"
-#include "Styling/SlateTypes.h"
+#include "Misc/Paths.h"
 #include "Styling/CoreStyle.h"
-#include "Styling/AppStyle.h"
 #include "Interfaces/IPluginManager.h"
 #include "SlateOptMacros.h"
+#include "Styling/SlateStyle.h"
 #include "Styling/SlateStyleMacros.h"
 #include "Styling/ToolBarStyle.h"
 
diff --git a/Engine/Plugins/Editor/ModelingToolsEditorMode/Source/ModelingToolsEditorMode/Private/ModelingToolsEditorModeToolkit.cpp b/Engine/Plugins/Editor/ModelingToolsEditorMode/Source/ModelingToolsEditorMode/Private/ModelingToolsEditorModeToolkit.cpp
--- a/Engine/Plugins/Editor/ModelingToolsEditorMode/Source/ModelingToolsEditorMode/Private/ModelingToolsEditorModeToolkit.cpp
+++ b/Engine/Plugins/Editor/ModelingToolsEditorMode/Source/ModelingToolsEditorMode/Private/ModelingToolsEditorModeToolkit.cpp
@@ -1,31 +1,24 @@
 // Copyright Epic Games, Inc. All Rights Reserved.
 
 #include "ModelingToolsEditorModeToolkit.h"
+#include "Brushes/SlateRoundedBoxBrush.h"
 #include "IGeometryProcessingInterfacesModule.h"
 #include "GeometryProcessingInterfaces/IUVEditorModularFeature.h"
-#include "ModelingToolsEditorMode.h"
+#include "EdModeInteractiveToolsContext.h"
 #include "ModelingToolsManagerActions.h"
+#include "InteractiveToolManager.h"
 #include "ModelingToolsEditorModeSettings.h"
 #include "ModelingToolsEditorModeStyle.h"
-#include "Selection/GeometrySelectionManager.h"
-#include "Engine/Selection.h"
 
-#include "Framework/Application/SlateApplication.h"
-#include "Framework/MultiBox/MultiBoxBuilder.h"
 #include "Modules/ModuleManager.h"
 #include "IDetailsView.h"
 #include "ISettingsModule.h"
-#include "EditorModeManager.h"
 #include "Toolkits/AssetEditorModeUILayer.h"
 
-#include "Widgets/Input/SButton.h"
 #include "SSimpleButton.h"
-#include "Widgets/Text/STextBlock.h"
-#include "Widgets/Layout/SSeparator.h"
-#include "Widgets/Images/SImage.h"
+#include "Tools/UEdMode.h"
 #include "Widgets/Layout/SExpandableArea.h"
 #include "Framework/MultiBox/SToolBarButtonBlock.h"
-#include "Internationalization/Text.h"
 #include "ModelingWidgets/ModelingCustomizationUtil.h"
 
 // for Tool Extensions
@@ -45,6 +38,8 @@
 #include "ModelingModeAssetUtils.h"
 
 #include "LevelEditorViewport.h"
+#include "Widgets/Input/SEditableTextBox.h"
+#include "Widgets/Input/STextComboBox.h"
 
 
 #define LOCTEXT_NAMESPACE "FModelingToolsEditorModeToolkit"
diff --git a/Engine/Plugins/Editor/ModelingToolsEditorMode/Source/ModelingToolsEditorMode/Private/ModelingToolsManagerActions.cpp b/Engine/Plugins/Editor/ModelingToolsEditorMode/Source/ModelingToolsEditorMode/Private/ModelingToolsManagerActions.cpp
--- a/Engine/Plugins/Editor/ModelingToolsEditorMode/Source/ModelingToolsEditorMode/Private/ModelingToolsManagerActions.cpp
+++ b/Engine/Plugins/Editor/ModelingToolsEditorMode/Source/ModelingToolsEditorMode/Private/ModelingToolsManagerActions.cpp
@@ -1,9 +1,9 @@
 // Copyright Epic Games, Inc. All Rights Reserved.
 
 #include "ModelingToolsManagerActions.h"
-#include "InputCoreTypes.h"
 #include "ModelingToolsEditorModeStyle.h"
 #include "ModelingToolsEditorModeSettings.h"
+#include "Styling/ISlateStyle.h"
 
 #define LOCTEXT_NAMESPACE "ModelingToolsManagerCommands"
 
diff --git a/Engine/Plugins/Editor/ObjectMixer/LightMixer/Source/LightMixer/Private/LightMixerModule.cpp b/Engine/Plugins/Editor/ObjectMixer/LightMixer/Source/LightMixer/Private/LightMixerModule.cpp
--- a/Engine/Plugins/Editor/ObjectMixer/LightMixer/Source/LightMixer/Private/LightMixerModule.cpp
+++ b/Engine/Plugins/Editor/ObjectMixer/LightMixer/Source/LightMixer/Private/LightMixerModule.cpp
@@ -7,7 +7,7 @@
 #include "LightMixerStyle.h"
 
 #include "ISettingsModule.h"
-#include "ISettingsSection.h"
+#include "Misc/CoreDelegates.h"
 #include "Modules/ModuleManager.h"
 
 #define LOCTEXT_NAMESPACE "FLightMixerEditorModule"
diff --git a/Engine/Plugins/Editor/ObjectMixer/LightMixer/Source/LightMixer/Private/LightMixerStyle.cpp b/Engine/Plugins/Editor/ObjectMixer/LightMixer/Source/LightMixer/Private/LightMixerStyle.cpp
--- a/Engine/Plugins/Editor/ObjectMixer/LightMixer/Source/LightMixer/Private/LightMixerStyle.cpp
+++ b/Engine/Plugins/Editor/ObjectMixer/LightMixer/Source/LightMixer/Private/LightMixerStyle.cpp
@@ -2,13 +2,12 @@
 
 #include "LightMixerStyle.h"
 
-#include "Styling/AppStyle.h"
+#include "Brushes/SlateImageBrush.h"
 #include "Framework/Application/SlateApplication.h"
 #include "Interfaces/IPluginManager.h"
-#include "Styling/CoreStyle.h"
+#include "Misc/Paths.h"
 #include "Styling/SlateStyleRegistry.h"
-#include "Styling/SlateTypes.h"
-#include "Styling/StyleColors.h"
+#include "Rendering/SlateRenderer.h"
 
 TSharedPtr<FSlateStyleSet> FLightMixerStyle::StyleInstance = nullptr;
 
diff --git a/Engine/Plugins/Editor/ObjectMixer/LightMixer/Source/LightMixer/Private/LightMixerStyle.h b/Engine/Plugins/Editor/ObjectMixer/LightMixer/Source/LightMixer/Private/LightMixerStyle.h
--- a/Engine/Plugins/Editor/ObjectMixer/LightMixer/Source/LightMixer/Private/LightMixerStyle.h
+++ b/Engine/Plugins/Editor/ObjectMixer/LightMixer/Source/LightMixer/Private/LightMixerStyle.h
@@ -4,6 +4,8 @@
 
 #include "Styling/SlateStyle.h"
 
+class ISlateStyle;
+
 class FLightMixerStyle : FSlateStyleSet
 {
 public:
diff --git a/Engine/Plugins/Editor/ObjectMixer/ObjectMixer/Source/ObjectMixer/Private/ObjectFilter/ObjectMixerEditorObjectFilter.cpp b/Engine/Plugins/Editor/ObjectMixer/ObjectMixer/Source/ObjectMixer/Private/ObjectFilter/ObjectMixerEditorObjectFilter.cpp
--- a/Engine/Plugins/Editor/ObjectMixer/ObjectMixer/Source/ObjectMixer/Private/ObjectFilter/ObjectMixerEditorObjectFilter.cpp
+++ b/Engine/Plugins/Editor/ObjectMixer/ObjectMixer/Source/ObjectMixer/Private/ObjectFilter/ObjectMixerEditorObjectFilter.cpp
@@ -2,7 +2,6 @@
 
 #include "ObjectFilter/ObjectMixerEditorObjectFilter.h"
 
-#include "GameFramework/Actor.h"
 #include "Kismet2/ComponentEditorUtils.h"
 
 FText UObjectMixerObjectFilter::GetRowDisplayName(UObject* InObject, const bool bIsHybridRow) const
diff --git a/Engine/Plugins/Editor/ObjectMixer/ObjectMixer/Source/ObjectMixer/Private/ObjectMixerEditorModule.cpp b/Engine/Plugins/Editor/ObjectMixer/ObjectMixer/Source/ObjectMixer/Private/ObjectMixerEditorModule.cpp
--- a/Engine/Plugins/Editor/ObjectMixer/ObjectMixer/Source/ObjectMixer/Private/ObjectMixerEditorModule.cpp
+++ b/Engine/Plugins/Editor/ObjectMixer/ObjectMixer/Source/ObjectMixer/Private/ObjectMixerEditorModule.cpp
@@ -2,20 +2,19 @@
 
 #include "ObjectMixerEditorModule.h"
 
-#include "ObjectMixerEditorLog.h"
+#include "Engine/Level.h"
 #include "ObjectMixerEditorSettings.h"
+#include "Misc/CoreDelegates.h"
 #include "ObjectMixerEditorStyle.h"
 #include "Views/MainPanel/ObjectMixerEditorMainPanel.h"
 
-#include "Editor.h"
-#include "Framework/Docking/TabManager.h"
 #include "ISettingsModule.h"
 #include "LevelEditor.h"
 #include "ObjectMixerEditorSerializedData.h"
 #include "Selection.h"
 #include "Misc/TransactionObjectEvent.h"
-#include "Modules/ModuleManager.h"
 #include "ToolMenus.h"
+#include "Widgets/Docking/SDockTab.h"
 #include "WorkspaceMenuStructure.h"
 #include "WorkspaceMenuStructureModule.h"
 #include "Views/Widgets/ObjectMixerEditorListMenuContext.h"
diff --git a/Engine/Plugins/Editor/ObjectMixer/ObjectMixer/Source/ObjectMixer/Private/ObjectMixerEditorStyle.cpp b/Engine/Plugins/Editor/ObjectMixer/ObjectMixer/Source/ObjectMixer/Private/ObjectMixerEditorStyle.cpp
--- a/Engine/Plugins/Editor/ObjectMixer/ObjectMixer/Source/ObjectMixer/Private/ObjectMixerEditorStyle.cpp
+++ b/Engine/Plugins/Editor/ObjectMixer/ObjectMixer/Source/ObjectMixer/Private/ObjectMixerEditorStyle.cpp
@@ -2,13 +2,12 @@
 
 #include "ObjectMixerEditorStyle.h"
 
-#include "Styling/AppStyle.h"
+#include "Brushes/SlateImageBrush.h"
 #include "Framework/Application/SlateApplication.h"
 #include "Interfaces/IPluginManager.h"
-#include "Styling/CoreStyle.h"
+#include "Misc/Paths.h"
 #include "Styling/SlateStyleRegistry.h"
-#include "Styling/SlateTypes.h"
-#include "Styling/StyleColors.h"
+#include "Rendering/SlateRenderer.h"
 
 TSharedPtr<FSlateStyleSet> FObjectMixerEditorStyle::StyleInstance = nullptr;
 
diff --git a/Engine/Plugins/Editor/ObjectMixer/ObjectMixer/Source/ObjectMixer/Private/ObjectMixerFilterFactory.cpp b/Engine/Plugins/Editor/ObjectMixer/ObjectMixer/Source/ObjectMixer/Private/ObjectMixerFilterFactory.cpp
--- a/Engine/Plugins/Editor/ObjectMixer/ObjectMixer/Source/ObjectMixer/Private/ObjectMixerFilterFactory.cpp
+++ b/Engine/Plugins/Editor/ObjectMixer/ObjectMixer/Source/ObjectMixer/Private/ObjectMixerFilterFactory.cpp
@@ -2,10 +2,12 @@
 
 #include "ObjectMixerFilterFactory.h"
 
+#include "AssetTypeCategories.h"
 #include "ObjectFilter/ObjectMixerEditorObjectFilter.h"
 
-#include "AssetToolsModule.h"
+#include "Engine/Blueprint.h"
 #include "Kismet2/KismetEditorUtilities.h"
+#include "Engine/BlueprintGeneratedClass.h"
 
 #define LOCTEXT_NAMESPACE "ObjectMixerFilterFactory"
 
diff --git a/Engine/Plugins/Editor/ObjectMixer/ObjectMixer/Source/ObjectMixer/Private/Views/List/ObjectMixerEditorList.cpp b/Engine/Plugins/Editor/ObjectMixer/ObjectMixer/Source/ObjectMixer/Private/Views/List/ObjectMixerEditorList.cpp
--- a/Engine/Plugins/Editor/ObjectMixer/ObjectMixer/Source/ObjectMixer/Private/Views/List/ObjectMixerEditorList.cpp
+++ b/Engine/Plugins/Editor/ObjectMixer/ObjectMixer/Source/ObjectMixer/Private/Views/List/ObjectMixerEditorList.cpp
@@ -3,6 +3,7 @@
 #include "Views/List/ObjectMixerEditorList.h"
 
 #include "Views/List/SObjectMixerEditorList.h"
+#include "Views/MainPanel/ObjectMixerEditorMainPanel.h"
 
 FObjectMixerEditorList::FObjectMixerEditorList(TSharedRef<FObjectMixerEditorMainPanel, ESPMode::ThreadSafe> InMainPanel)
 {
diff --git a/Engine/Plugins/Editor/ObjectMixer/ObjectMixer/Source/ObjectMixer/Private/Views/List/ObjectMixerEditorListRow.cpp b/Engine/Plugins/Editor/ObjectMixer/ObjectMixer/Source/ObjectMixer/Private/Views/List/ObjectMixerEditorListRow.cpp
--- a/Engine/Plugins/Editor/ObjectMixer/ObjectMixer/Source/ObjectMixer/Private/Views/List/ObjectMixerEditorListRow.cpp
+++ b/Engine/Plugins/Editor/ObjectMixer/ObjectMixer/Source/ObjectMixer/Private/Views/List/ObjectMixerEditorListRow.cpp
@@ -8,8 +8,8 @@
 #include "ObjectMixerEditorSerializedData.h"
 #include "ScopedTransaction.h"
 #include "GameFramework/Actor.h"
-#include "Styling/SlateIconFinder.h"
 #include "Views/List/SObjectMixerEditorList.h"
+#include "Views/MainPanel/ObjectMixerEditorMainPanel.h"
 
 TSharedRef<FObjectMixerListRowDragDropOp> FObjectMixerListRowDragDropOp::New(const TArray<FObjectMixerEditorListRowPtr>& InItems)
 {
diff --git a/Engine/Plugins/Editor/ObjectMixer/ObjectMixer/Source/ObjectMixer/Private/Views/List/SObjectMixerEditorList.cpp b/Engine/Plugins/Editor/ObjectMixer/ObjectMixer/Source/ObjectMixer/Private/Views/List/SObjectMixerEditorList.cpp
--- a/Engine/Plugins/Editor/ObjectMixer/ObjectMixer/Source/ObjectMixer/Private/Views/List/SObjectMixerEditorList.cpp
+++ b/Engine/Plugins/Editor/ObjectMixer/ObjectMixer/Source/ObjectMixer/Private/Views/List/SObjectMixerEditorList.cpp
@@ -3,37 +3,32 @@
 #include "Views/List/SObjectMixerEditorList.h"
 
 // View Filters
+#include "Algo/Find.h"
 #include "Views/List/ObjectMixerEditorListFilters/IObjectMixerEditorListFilter.h"
 
-#include "ObjectFilter/ObjectMixerEditorObjectFilter.h"
+#include "Engine/Level.h"
 #include "ObjectMixerEditorLog.h"
+#include "Framework/Application/SlateApplication.h"
 #include "ObjectMixerEditorSettings.h"
+#include "Framework/Views/TableViewMetadata.h"
 #include "ObjectMixerEditorStyle.h"
+#include "Views/List/ObjectMixerEditorList.h"
 #include "Views/List/SObjectMixerEditorListRow.h"
 
-#include "ActorBrowsingMode.h"
 #include "Algo/AllOf.h"
 #include "Algo/AnyOf.h"
-#include "Algo/IndexOf.h"
-#include "DragAndDrop/AssetDragDropOp.h"
-#include "Editor.h"
-#include "Editor/UnrealEdEngine.h"
 #include "EditorActorFolders.h"
 #include "Framework/MultiBox/MultiBoxBuilder.h"
-#include "IPlacementModeModule.h"
-#include "LevelEditor.h"
-#include "LevelEditorMenuContext.h"
 #include "LevelEditorViewport.h"
 #include "ObjectMixerEditorModule.h"
 #include "ObjectMixerEditorSerializedData.h"
 #include "ScopedTransaction.h"
 #include "Selection.h"
-#include "Styling/StyleColors.h"
+#include "Views/MainPanel/ObjectMixerEditorMainPanel.h"
 #include "Views/Widgets/ObjectMixerEditorListMenuContext.h"
-#include "Widgets/Layout/SScrollBox.h"
 #include "Widgets/Layout/SWidgetSwitcher.h"
 #include "Widgets/Text/SRichTextBlock.h"
-#include "WorldPartition/WorldPartition.h"
+#include "Widgets/Views/STreeView.h"
 
 #define LOCTEXT_NAMESPACE "ObjectMixerEditorList"
 
diff --git a/Engine/Plugins/Editor/ObjectMixer/ObjectMixer/Source/ObjectMixer/Private/Views/Widgets/SCollectionSelectionButton.cpp b/Engine/Plugins/Editor/ObjectMixer/ObjectMixer/Source/ObjectMixer/Private/Views/Widgets/SCollectionSelectionButton.cpp
--- a/Engine/Plugins/Editor/ObjectMixer/ObjectMixer/Source/ObjectMixer/Private/Views/Widgets/SCollectionSelectionButton.cpp
+++ b/Engine/Plugins/Editor/ObjectMixer/ObjectMixer/Source/ObjectMixer/Private/Views/Widgets/SCollectionSelectionButton.cpp
@@ -2,18 +2,13 @@
 
 #include "Views/Widgets/SCollectionSelectionButton.h"
 
-#include "ObjectMixerEditorSerializedData.h"
 #include "Views/List/ObjectMixerEditorListFilters/ObjectMixerEditorListFilter_Collection.h"
 #include "Views/MainPanel/ObjectMixerEditorMainPanel.h"
 #include "Views/MainPanel/SObjectMixerEditorMainPanel.h"
 
-#include "DragAndDrop/DecoratedDragDropOp.h"
 #include "Framework/Commands/GenericCommands.h"
 #include "Framework/MultiBox/MultiBoxBuilder.h"
-#include "Styling/AppStyle.h"
-#include "Views/List/ObjectMixerEditorListRow.h"
 #include "Widgets/Input/SMenuAnchor.h"
-#include "Widgets/Layout/SBorder.h"
 #include "Widgets/Text/SInlineEditableTextBlock.h"
 
 #define LOCTEXT_NAMESPACE "ObjectMixerEditor"
diff --git a/Engine/Plugins/Editor/ObjectMixer/ObjectMixer/Source/ObjectMixer/Private/Views/Widgets/SObjectMixerPlacementAssetMenuEntry.cpp b/Engine/Plugins/Editor/ObjectMixer/ObjectMixer/Source/ObjectMixer/Private/Views/Widgets/SObjectMixerPlacementAssetMenuEntry.cpp
--- a/Engine/Plugins/Editor/ObjectMixer/ObjectMixer/Source/ObjectMixer/Private/Views/Widgets/SObjectMixerPlacementAssetMenuEntry.cpp
+++ b/Engine/Plugins/Editor/ObjectMixer/ObjectMixer/Source/ObjectMixer/Private/Views/Widgets/SObjectMixerPlacementAssetMenuEntry.cpp
@@ -2,7 +2,6 @@
 
 #include "SObjectMixerPlacementAssetMenuEntry.h"
 
-#include "Application/SlateApplicationBase.h"
 #include "AssetSelection.h"
 #include "ClassIconFinder.h"
 #include "DragAndDrop/AssetDragDropOp.h"
@@ -11,10 +10,9 @@
 #include "IPlacementModeModule.h"
 #include "LevelEditorActions.h"
 #include "LevelEditorViewport.h"
-#include "Styling/SlateIconFinder.h"
-#include "Styling/StyleColors.h"
 #include "Widgets/Layout/SBox.h"
 #include "Widgets/Layout/SSpacer.h"
+#include "Widgets/SToolTip.h"
 
 void SObjectMixerPlacementAssetMenuEntry::Construct(const FArguments& InArgs, const TSharedPtr<const FPlaceableItem>& InItem)
 {	
diff --git a/Engine/Plugins/Editor/ObjectMixer/ObjectMixer/Source/ObjectMixer/Private/Views/Widgets/SObjectMixerPlacementAssetMenuEntry.h b/Engine/Plugins/Editor/ObjectMixer/ObjectMixer/Source/ObjectMixer/Private/Views/Widgets/SObjectMixerPlacementAssetMenuEntry.h
--- a/Engine/Plugins/Editor/ObjectMixer/ObjectMixer/Source/ObjectMixer/Private/Views/Widgets/SObjectMixerPlacementAssetMenuEntry.h
+++ b/Engine/Plugins/Editor/ObjectMixer/ObjectMixer/Source/ObjectMixer/Private/Views/Widgets/SObjectMixerPlacementAssetMenuEntry.h
@@ -2,10 +2,11 @@
 
 #pragma once
 
-#include "Styling/SlateBrush.h"
-#include "Styling/SlateTypes.h"
 #include "Widgets/SCompoundWidget.h"
 
+struct FButtonStyle;
+struct FPlaceableItem;
+
 class SObjectMixerPlacementAssetMenuEntry : public SCompoundWidget
 {
 public:
diff --git a/Engine/Plugins/Editor/PluginBrowser/Source/PluginBrowser/Private/DefaultPluginWizardDefinition.cpp b/Engine/Plugins/Editor/PluginBrowser/Source/PluginBrowser/Private/DefaultPluginWizardDefinition.cpp
--- a/Engine/Plugins/Editor/PluginBrowser/Source/PluginBrowser/Private/DefaultPluginWizardDefinition.cpp
+++ b/Engine/Plugins/Editor/PluginBrowser/Source/PluginBrowser/Private/DefaultPluginWizardDefinition.cpp
@@ -2,10 +2,10 @@
 
 #include "DefaultPluginWizardDefinition.h"
 #include "GenericPlatform/GenericPlatformFile.h"
-#include "PluginStyle.h"
 #include "Interfaces/IPluginManager.h"
 #include "HAL/PlatformFileManager.h"
 #include "Misc/App.h"
+#include "Misc/Paths.h"
 #include "PluginBrowserModule.h"
 
 #define LOCTEXT_NAMESPACE "NewPluginWizard"
@@ -192,4 +192,4 @@ FString FDefaultPluginWizardDefinition::GetFolderForTemplate(TSharedRef<FPluginT
 	return InTemplate->OnDiskPath;
 }
 
-#undef LOCTEXT_NAMESPACE
\ No newline at end of file
+#undef LOCTEXT_NAMESPACE
diff --git a/Engine/Plugins/Editor/PluginBrowser/Source/PluginBrowser/Private/DefaultPluginWizardDefinition.h b/Engine/Plugins/Editor/PluginBrowser/Source/PluginBrowser/Private/DefaultPluginWizardDefinition.h
--- a/Engine/Plugins/Editor/PluginBrowser/Source/PluginBrowser/Private/DefaultPluginWizardDefinition.h
+++ b/Engine/Plugins/Editor/PluginBrowser/Source/PluginBrowser/Private/DefaultPluginWizardDefinition.h
@@ -3,6 +3,7 @@
 #pragma once
 
 #include "IPluginWizardDefinition.h"
+#include "Styling/SlateTypes.h"
 
 class FDefaultPluginWizardDefinition : public IPluginWizardDefinition
 {
@@ -52,4 +53,4 @@ private:
 
 	/** If true, this definition is for a project that can only contain content */
 	bool bIsContentOnlyProject;
-};
\ No newline at end of file
+};
diff --git a/Engine/Plugins/Editor/PluginBrowser/Source/PluginBrowser/Private/NewPluginDescriptorData.h b/Engine/Plugins/Editor/PluginBrowser/Source/PluginBrowser/Private/NewPluginDescriptorData.h
--- a/Engine/Plugins/Editor/PluginBrowser/Source/PluginBrowser/Private/NewPluginDescriptorData.h
+++ b/Engine/Plugins/Editor/PluginBrowser/Source/PluginBrowser/Private/NewPluginDescriptorData.h
@@ -2,10 +2,6 @@
 
 #pragma once
 
-#include "CoreMinimal.h"
-#include "UObject/ObjectMacros.h"
-#include "UObject/UObjectGlobals.h"
-#include "UObject/Object.h"
 #include "NewPluginDescriptorData.generated.h"
 
 /**
@@ -37,4 +33,4 @@ public:
 	UPROPERTY(EditAnywhere, Category="Descriptor Data", AdvancedDisplay)
 	bool bIsBetaVersion;
 
-};
\ No newline at end of file
+};
diff --git a/Engine/Plugins/Editor/PluginBrowser/Source/PluginBrowser/Private/PluginBrowserModule.cpp b/Engine/Plugins/Editor/PluginBrowser/Source/PluginBrowser/Private/PluginBrowserModule.cpp
--- a/Engine/Plugins/Editor/PluginBrowser/Source/PluginBrowser/Private/PluginBrowserModule.cpp
+++ b/Engine/Plugins/Editor/PluginBrowser/Source/PluginBrowser/Private/PluginBrowserModule.cpp
@@ -1,8 +1,6 @@
 // Copyright Epic Games, Inc. All Rights Reserved.
 
 #include "PluginBrowserModule.h"
-#include "Textures/SlateIcon.h"
-#include "Framework/Docking/TabManager.h"
 #include "SPluginBrowser.h"
 #include "Features/IModularFeatures.h"
 #include "Features/EditorFeatures.h"
@@ -17,7 +15,6 @@
 #include "Misc/ConfigCacheIni.h"
 #include "Interfaces/IPluginManager.h"
 #include "PluginDescriptorEditor.h"
-#include "HAL/LowLevelMemTracker.h"
 
 #define LOCTEXT_NAMESPACE "PluginsEditor"
 
diff --git a/Engine/Plugins/Editor/PluginBrowser/Source/PluginBrowser/Private/PluginBrowserModule.h b/Engine/Plugins/Editor/PluginBrowser/Source/PluginBrowser/Private/PluginBrowserModule.h
--- a/Engine/Plugins/Editor/PluginBrowser/Source/PluginBrowser/Private/PluginBrowserModule.h
+++ b/Engine/Plugins/Editor/PluginBrowser/Source/PluginBrowser/Private/PluginBrowserModule.h
@@ -2,11 +2,13 @@
 
 #pragma once 
 
-#include "CoreMinimal.h"
-#include "SlateFwd.h"
-#include "Modules/ModuleManager.h"
 #include "IPluginBrowser.h"
-#include "Widgets/SWindow.h"
+
+class IModuleInterface;
+class IPluginWizardDefinition;
+class SDockTab;
+class SNotificationItem;
+class SWindow;
 
 class FSpawnTabArgs;
 
diff --git a/Engine/Plugins/Editor/PluginBrowser/Source/PluginBrowser/Private/PluginBrowserPrivate.h b/Engine/Plugins/Editor/PluginBrowser/Source/PluginBrowser/Private/PluginBrowserPrivate.h
--- a/Engine/Plugins/Editor/PluginBrowser/Source/PluginBrowser/Private/PluginBrowserPrivate.h
+++ b/Engine/Plugins/Editor/PluginBrowser/Source/PluginBrowser/Private/PluginBrowserPrivate.h
@@ -3,8 +3,8 @@
 #pragma once
 
 
-#include "CoreMinimal.h"
-#include "Misc/TextFilter.h"
+
+template <typename ItemType> class TTextFilter;
 
 class IPlugin;
 
diff --git a/Engine/Plugins/Editor/PluginBrowser/Source/PluginBrowser/Private/PluginDescriptorEditor.cpp b/Engine/Plugins/Editor/PluginBrowser/Source/PluginBrowser/Private/PluginDescriptorEditor.cpp
--- a/Engine/Plugins/Editor/PluginBrowser/Source/PluginBrowser/Private/PluginDescriptorEditor.cpp
+++ b/Engine/Plugins/Editor/PluginBrowser/Source/PluginBrowser/Private/PluginDescriptorEditor.cpp
@@ -1,25 +1,21 @@
 // Copyright Epic Games, Inc. All Rights Reserved.
 
 #include "PluginDescriptorEditor.h"
-#include "HAL/PlatformFileManager.h"
+#include "DetailsViewArgs.h"
 #include "Misc/MessageDialog.h"
-#include "Misc/App.h"
-#include "Modules/ModuleManager.h"
+#include "IDetailsView.h"
 #include "Widgets/SWindow.h"
 #include "Framework/Application/SlateApplication.h"
-#include "Widgets/Layout/SBorder.h"
+#include "LocalizationDescriptor.h"
 #include "Widgets/Text/STextBlock.h"
+#include "PluginReferenceDescriptor.h"
 #include "Widgets/Input/SButton.h"
-#include "Styling/AppStyle.h"
-#include "ISourceControlOperation.h"
 #include "SourceControlOperations.h"
 #include "ISourceControlProvider.h"
 #include "ISourceControlModule.h"
-#include "PluginDescriptor.h"
 #include "Interfaces/IPluginManager.h"
 #include "PluginStyle.h"
 #include "PluginMetadataObject.h"
-#include "Interfaces/IProjectManager.h"
 #include "PropertyEditorModule.h"
 #include "Misc/ScopedSlowTask.h"
 
diff --git a/Engine/Plugins/Editor/PluginBrowser/Source/PluginBrowser/Private/PluginDescriptorEditor.h b/Engine/Plugins/Editor/PluginBrowser/Source/PluginBrowser/Private/PluginDescriptorEditor.h
--- a/Engine/Plugins/Editor/PluginBrowser/Source/PluginBrowser/Private/PluginDescriptorEditor.h
+++ b/Engine/Plugins/Editor/PluginBrowser/Source/PluginBrowser/Private/PluginDescriptorEditor.h
@@ -2,8 +2,9 @@
 
 #pragma once
 
-#include "CoreMinimal.h"
-#include "Input/Reply.h"
+#include "Delegates/Delegate.h"
+
+class FReply;
 
 class UPluginMetadataObject;
 class IPlugin;
diff --git a/Engine/Plugins/Editor/PluginBrowser/Source/PluginBrowser/Private/PluginMetadataObject.cpp b/Engine/Plugins/Editor/PluginBrowser/Source/PluginBrowser/Private/PluginMetadataObject.cpp
--- a/Engine/Plugins/Editor/PluginBrowser/Source/PluginBrowser/Private/PluginMetadataObject.cpp
+++ b/Engine/Plugins/Editor/PluginBrowser/Source/PluginBrowser/Private/PluginMetadataObject.cpp
@@ -2,18 +2,13 @@
 
 #include "PluginMetadataObject.h"
 #include "Misc/Paths.h"
-#include "Layout/Margin.h"
-#include "Widgets/DeclarativeSyntaxSupport.h"
 #include "Widgets/SBoxPanel.h"
 #include "Widgets/Text/STextBlock.h"
 #include "Interfaces/IPluginManager.h"
 #include "PluginReferenceDescriptor.h"
 #include "DetailLayoutBuilder.h"
 #include "DetailWidgetRow.h"
-#include "DetailCategoryBuilder.h"
 #include "SExternalImageReference.h"
-#include "Dom/JsonObject.h"
-#include "Features/IPluginsEditorFeature.h"
 #include "PluginBrowserModule.h"
 
 #include UE_INLINE_GENERATED_CPP_BY_NAME(PluginMetadataObject)
diff --git a/Engine/Plugins/Editor/PluginBrowser/Source/PluginBrowser/Private/PluginStyle.cpp b/Engine/Plugins/Editor/PluginBrowser/Source/PluginBrowser/Private/PluginStyle.cpp
--- a/Engine/Plugins/Editor/PluginBrowser/Source/PluginBrowser/Private/PluginStyle.cpp
+++ b/Engine/Plugins/Editor/PluginBrowser/Source/PluginBrowser/Private/PluginStyle.cpp
@@ -1,12 +1,16 @@
 // Copyright Epic Games, Inc. All Rights Reserved.
 
 #include "PluginStyle.h"
+#include "Brushes/SlateImageBrush.h"
 #include "Styling/SlateStyleRegistry.h"
+#include "Brushes/SlateNoResource.h"
 #include "Styling/SlateTypes.h"
+#include "Brushes/SlateRoundedBoxBrush.h"
 #include "Styling/CoreStyle.h"
-#include "Styling/AppStyle.h"
 #include "Interfaces/IPluginManager.h"
+#include "Misc/Paths.h"
 #include "Styling/StarshipCoreStyle.h"
+#include "Styling/SlateStyle.h"
 
 #define RootToContentDir FPluginStyle::InContent
 #define RootToCoreContentDir StyleSet->RootToCoreContentDir
@@ -185,4 +189,4 @@ void FPluginStyle::Shutdown()
 	}
 }
 
-#undef RootToContentDir
\ No newline at end of file
+#undef RootToContentDir
diff --git a/Engine/Plugins/Editor/PluginBrowser/Source/PluginBrowser/Private/SFilePathBlock.cpp b/Engine/Plugins/Editor/PluginBrowser/Source/PluginBrowser/Private/SFilePathBlock.cpp
--- a/Engine/Plugins/Editor/PluginBrowser/Source/PluginBrowser/Private/SFilePathBlock.cpp
+++ b/Engine/Plugins/Editor/PluginBrowser/Source/PluginBrowser/Private/SFilePathBlock.cpp
@@ -2,10 +2,7 @@
 
 #include "SFilePathBlock.h"
 #include "Widgets/SOverlay.h"
-#include "Widgets/Layout/SBorder.h"
-#include "Widgets/Text/STextBlock.h"
 #include "Widgets/Layout/SBox.h"
-#include "Widgets/Layout/SGridPanel.h"
 #include "Widgets/Input/SEditableTextBox.h"
 #include "Widgets/Input/SButton.h"
 #include "Widgets/SBoxPanel.h"
diff --git a/Engine/Plugins/Editor/PluginBrowser/Source/PluginBrowser/Private/SFilePathBlock.h b/Engine/Plugins/Editor/PluginBrowser/Source/PluginBrowser/Private/SFilePathBlock.h
--- a/Engine/Plugins/Editor/PluginBrowser/Source/PluginBrowser/Private/SFilePathBlock.h
+++ b/Engine/Plugins/Editor/PluginBrowser/Source/PluginBrowser/Private/SFilePathBlock.h
@@ -2,12 +2,8 @@
 
 #pragma once
 
-#include "CoreMinimal.h"
-#include "Styling/SlateColor.h"
-#include "Widgets/DeclarativeSyntaxSupport.h"
 #include "Widgets/SCompoundWidget.h"
 #include "Framework/SlateDelegates.h"
-#include "Styling/AppStyle.h"
 
 class SEditableTextBox;
 
diff --git a/Engine/Plugins/Editor/PluginBrowser/Source/PluginBrowser/Private/SNewPluginWizard.cpp b/Engine/Plugins/Editor/PluginBrowser/Source/PluginBrowser/Private/SNewPluginWizard.cpp
--- a/Engine/Plugins/Editor/PluginBrowser/Source/PluginBrowser/Private/SNewPluginWizard.cpp
+++ b/Engine/Plugins/Editor/PluginBrowser/Source/PluginBrowser/Private/SNewPluginWizard.cpp
@@ -1,31 +1,23 @@
 // Copyright Epic Games, Inc. All Rights Reserved.
 
 #include "SNewPluginWizard.h"
+#include "Brushes/SlateDynamicImageBrush.h"
 #include "Misc/Paths.h"
-#include "HAL/PlatformFileManager.h"
+#include "DetailsViewArgs.h"
 #include "Misc/MessageDialog.h"
+#include "Framework/Views/TableViewMetadata.h"
 #include "HAL/FileManager.h"
+#include "IDesktopPlatform.h"
 #include "Misc/App.h"
-#include "Widgets/SBoxPanel.h"
-#include "Styling/SlateTypes.h"
 #include "Framework/Application/SlateApplication.h"
-#include "Widgets/Layout/SBorder.h"
-#include "Widgets/Images/SImage.h"
-#include "Widgets/Text/STextBlock.h"
+#include "LocalizationDescriptor.h"
 #include "Widgets/Text/SRichTextBlock.h"
-#include "Widgets/Layout/SBox.h"
-#include "Widgets/Input/SButton.h"
-#include "Widgets/Input/SCheckBox.h"
-#include "Styling/AppStyle.h"
-#include "ModuleDescriptor.h"
-#include "PluginDescriptor.h"
-#include "Interfaces/IPluginManager.h"
-#include "Widgets/Views/STableViewBase.h"
-#include "Widgets/Views/STableRow.h"
+#include "NewPluginDescriptorData.h"
 #include "Widgets/Views/SListView.h"
-#include "Widgets/Views/STileView.h"
+#include "PluginReferenceDescriptor.h"
 #include "PluginStyle.h"
 #include "DesktopPlatformModule.h"
+#include "Rendering/SlateRenderer.h"
 #include "Widgets/Docking/SDockTab.h"
 #include "Framework/Notifications/NotificationManager.h"
 #include "Widgets/Notifications/SNotificationList.h"
@@ -38,7 +30,6 @@
 #include "Interfaces/IMainFrameModule.h"
 #include "GameProjectGenerationModule.h"
 #include "DefaultPluginWizardDefinition.h"
-#include "UnrealEdMisc.h"
 #include "PropertyEditorModule.h"
 #include "IDetailsView.h"
 #include "Framework/Notifications/NotificationManager.h"
diff --git a/Engine/Plugins/Editor/PluginBrowser/Source/PluginBrowser/Private/SNewPluginWizard.h b/Engine/Plugins/Editor/PluginBrowser/Source/PluginBrowser/Private/SNewPluginWizard.h
--- a/Engine/Plugins/Editor/PluginBrowser/Source/PluginBrowser/Private/SNewPluginWizard.h
+++ b/Engine/Plugins/Editor/PluginBrowser/Source/PluginBrowser/Private/SNewPluginWizard.h
@@ -2,15 +2,13 @@
 
 #pragma once
 
-#include "CoreMinimal.h"
-#include "SlateFwd.h"
-#include "Input/Reply.h"
-#include "Brushes/SlateDynamicImageBrush.h"
-#include "Widgets/DeclarativeSyntaxSupport.h"
 #include "Widgets/SCompoundWidget.h"
-#include "IPluginWizardDefinition.h"
-#include "NewPluginDescriptorData.h"
-#include "ModuleDescriptor.h"
+
+class IPluginWizardDefinition;
+class SDockTab;
+class UNewPluginDescriptorData;
+namespace ESelectInfo { enum Type : int; }
+struct FPluginTemplateDescription;
 
 class ITableRow;
 class SFilePathBlock;
diff --git a/Engine/Plugins/Editor/PluginBrowser/Source/PluginBrowser/Private/SPluginBrowser.cpp b/Engine/Plugins/Editor/PluginBrowser/Source/PluginBrowser/Private/SPluginBrowser.cpp
--- a/Engine/Plugins/Editor/PluginBrowser/Source/PluginBrowser/Private/SPluginBrowser.cpp
+++ b/Engine/Plugins/Editor/PluginBrowser/Source/PluginBrowser/Private/SPluginBrowser.cpp
@@ -3,29 +3,23 @@
 #include "SPluginBrowser.h"
 #include "Framework/Application/SlateApplication.h"
 #include "Misc/App.h"
-#include "Modules/ModuleManager.h"
-#include "Widgets/Images/SImage.h"
+#include "SPluginCategory.h"
 #include "Widgets/Input/SComboButton.h"
-#include "Widgets/Text/STextBlock.h"
-#include "Widgets/Input/SButton.h"
+#include "SPluginTileList.h"
 #include "Widgets/SToolTip.h"
 #include "Framework/Docking/TabManager.h"
-#include "Styling/AppStyle.h"
 #include "UnrealEdMisc.h"
 #include "Interfaces/IPluginManager.h"
 #include "PluginStyle.h"
 #include "Widgets/Navigation/SBreadcrumbTrail.h"
 #include "SPluginCategoryTree.h"
-#include "SPluginTile.h"
 #include "Widgets/Input/SSearchBox.h"
 #include "PluginBrowserModule.h"
 #include "IDirectoryWatcher.h"
 #include "DirectoryWatcherModule.h"
 #include "Interfaces/IProjectManager.h"
 #include "ProjectDescriptor.h"
-#include "GameProjectGenerationModule.h"
 #include "Framework/MultiBox/MultiBoxBuilder.h"
-#include "SPrimaryButton.h"
 #include "Styling/StyleColors.h"
 #include "SWarningOrErrorBox.h"
 
diff --git a/Engine/Plugins/Editor/PluginBrowser/Source/PluginBrowser/Private/SPluginBrowser.h b/Engine/Plugins/Editor/PluginBrowser/Source/PluginBrowser/Private/SPluginBrowser.h
--- a/Engine/Plugins/Editor/PluginBrowser/Source/PluginBrowser/Private/SPluginBrowser.h
+++ b/Engine/Plugins/Editor/PluginBrowser/Source/PluginBrowser/Private/SPluginBrowser.h
@@ -2,15 +2,12 @@
 
 #pragma once
 
-#include "CoreMinimal.h"
-#include "SlateFwd.h"
-#include "Layout/Visibility.h"
-#include "Widgets/DeclarativeSyntaxSupport.h"
-#include "Input/Reply.h"
+#include "Misc/TextFilter.h"
 #include "Widgets/SCompoundWidget.h"
-#include "SPluginTileList.h"
 #include "PluginBrowserPrivate.h"
-#include "SPluginCategory.h"
+
+class FPluginCategory;
+template <typename ItemType> class SBreadcrumbTrail;
 
 class FActiveTimerHandle;
 class SPluginCategoryTree;
diff --git a/Engine/Plugins/Editor/PluginBrowser/Source/PluginBrowser/Private/SPluginCategory.cpp b/Engine/Plugins/Editor/PluginBrowser/Source/PluginBrowser/Private/SPluginCategory.cpp
--- a/Engine/Plugins/Editor/PluginBrowser/Source/PluginBrowser/Private/SPluginCategory.cpp
+++ b/Engine/Plugins/Editor/PluginBrowser/Source/PluginBrowser/Private/SPluginCategory.cpp
@@ -1,12 +1,8 @@
 // Copyright Epic Games, Inc. All Rights Reserved.
 
 #include "SPluginCategory.h"
-#include "Widgets/SBoxPanel.h"
-#include "Widgets/Layout/SBorder.h"
-#include "Widgets/Images/SImage.h"
-#include "Widgets/Text/STextBlock.h"
-#include "Widgets/Layout/SBox.h"
 #include "PluginStyle.h"
+#include "Widgets/Views/SListView.h"
 
 
 #define LOCTEXT_NAMESPACE "PluginCategoryTreeItem"
diff --git a/Engine/Plugins/Editor/PluginBrowser/Source/PluginBrowser/Private/SPluginCategory.h b/Engine/Plugins/Editor/PluginBrowser/Source/PluginBrowser/Private/SPluginCategory.h
--- a/Engine/Plugins/Editor/PluginBrowser/Source/PluginBrowser/Private/SPluginCategory.h
+++ b/Engine/Plugins/Editor/PluginBrowser/Source/PluginBrowser/Private/SPluginCategory.h
@@ -2,11 +2,10 @@
 
 #pragma once
 
-#include "CoreMinimal.h"
-#include "Widgets/DeclarativeSyntaxSupport.h"
-#include "Widgets/SCompoundWidget.h"
+#include "Widgets/Views/ITableRow.h"
 #include "Widgets/Views/STableRow.h"
-#include "Widgets/Views/SListView.h"
+
+class STableViewBase;
 
 class IPlugin;
 struct FSlateBrush;
diff --git a/Engine/Plugins/Editor/PluginBrowser/Source/PluginBrowser/Private/SPluginCategoryTree.cpp b/Engine/Plugins/Editor/PluginBrowser/Source/PluginBrowser/Private/SPluginCategoryTree.cpp
--- a/Engine/Plugins/Editor/PluginBrowser/Source/PluginBrowser/Private/SPluginCategoryTree.cpp
+++ b/Engine/Plugins/Editor/PluginBrowser/Source/PluginBrowser/Private/SPluginCategoryTree.cpp
@@ -1,9 +1,11 @@
 // Copyright Epic Games, Inc. All Rights Reserved.
 
 #include "SPluginCategoryTree.h"
+#include "Framework/Views/TableViewMetadata.h"
 #include "Interfaces/IPluginManager.h"
 #include "SPluginCategory.h"
 #include "SPluginBrowser.h"
+#include "Widgets/Views/STreeView.h"
 
 #define LOCTEXT_NAMESPACE "PluginCategories"
 
diff --git a/Engine/Plugins/Editor/PluginBrowser/Source/PluginBrowser/Private/SPluginCategoryTree.h b/Engine/Plugins/Editor/PluginBrowser/Source/PluginBrowser/Private/SPluginCategoryTree.h
--- a/Engine/Plugins/Editor/PluginBrowser/Source/PluginBrowser/Private/SPluginCategoryTree.h
+++ b/Engine/Plugins/Editor/PluginBrowser/Source/PluginBrowser/Private/SPluginCategoryTree.h
@@ -2,12 +2,12 @@
 
 #pragma once
 
-#include "CoreMinimal.h"
-#include "Widgets/DeclarativeSyntaxSupport.h"
 #include "Widgets/SCompoundWidget.h"
-#include "Widgets/Views/STableViewBase.h"
-#include "Widgets/Views/STableRow.h"
-#include "Widgets/Views/STreeView.h"
+
+class ITableRow;
+class STableViewBase;
+namespace ESelectInfo { enum Type : int; }
+template <typename ItemType> class STreeView;
 
 class FPluginCategory;
 class SPluginBrowser;
diff --git a/Engine/Plugins/Editor/PluginBrowser/Source/PluginBrowser/Private/SPluginTile.cpp b/Engine/Plugins/Editor/PluginBrowser/Source/PluginBrowser/Private/SPluginTile.cpp
--- a/Engine/Plugins/Editor/PluginBrowser/Source/PluginBrowser/Private/SPluginTile.cpp
+++ b/Engine/Plugins/Editor/PluginBrowser/Source/PluginBrowser/Private/SPluginTile.cpp
@@ -1,37 +1,31 @@
 // Copyright Epic Games, Inc. All Rights Reserved.
 
 #include "SPluginTile.h"
+#include "Brushes/SlateDynamicImageBrush.h"
 #include "HAL/PlatformFileManager.h"
+#include "HAL/PlatformFile.h"
 #include "Misc/MessageDialog.h"
+#include "IDesktopPlatform.h"
 #include "Misc/App.h"
-#include "Modules/ModuleManager.h"
-#include "Widgets/SWindow.h"
 #include "Framework/Application/SlateApplication.h"
+#include "PluginReferenceDescriptor.h"
 #include "Widgets/Images/SImage.h"
-#include "Widgets/Text/STextBlock.h"
-#include "Widgets/Input/SButton.h"
+#include "Rendering/SlateRenderer.h"
 #include "Widgets/Input/SCheckBox.h"
-#include "Styling/AppStyle.h"
-#include "ISourceControlOperation.h"
-#include "SourceControlOperations.h"
-#include "ISourceControlProvider.h"
-#include "ISourceControlModule.h"
-#include "PluginDescriptor.h"
 #include "Interfaces/IPluginManager.h"
 #include "SPluginBrowser.h"
 #include "PluginStyle.h"
 #include "GameProjectGenerationModule.h"
-#include "IDetailsView.h"
+#include "SPluginTileList.h"
 #include "Widgets/Input/SHyperlink.h"
-#include "PluginMetadataObject.h"
 #include "Interfaces/IProjectManager.h"
 #include "PluginBrowserModule.h"
-#include "PropertyEditorModule.h"
 #include "IUATHelperModule.h"
 #include "DesktopPlatformModule.h"
 #include "Misc/ConfigCacheIni.h"
+#include "Widgets/Layout/SBox.h"
 #include "Widgets/Layout/SSpacer.h"
-#include "Styling/StyleColors.h"
+#include "Widgets/SBoxPanel.h"
 
 #define LOCTEXT_NAMESPACE "PluginListTile"
 
diff --git a/Engine/Plugins/Editor/PluginBrowser/Source/PluginBrowser/Private/SPluginTile.h b/Engine/Plugins/Editor/PluginBrowser/Source/PluginBrowser/Private/SPluginTile.h
--- a/Engine/Plugins/Editor/PluginBrowser/Source/PluginBrowser/Private/SPluginTile.h
+++ b/Engine/Plugins/Editor/PluginBrowser/Source/PluginBrowser/Private/SPluginTile.h
@@ -2,13 +2,10 @@
 
 #pragma once
 
-#include "CoreMinimal.h"
-#include "Layout/Visibility.h"
-#include "Widgets/DeclarativeSyntaxSupport.h"
-#include "Input/Reply.h"
 #include "Widgets/SCompoundWidget.h"
-#include "SPluginTileList.h"
-#include "Brushes/SlateDynamicImageBrush.h"
+
+enum class ECheckBoxState : uint8;
+struct FSlateDynamicImageBrush;
 
 class IPlugin;
 class UPluginMetadataObject;
diff --git a/Engine/Plugins/Editor/PluginBrowser/Source/PluginBrowser/Private/SPluginTileList.cpp b/Engine/Plugins/Editor/PluginBrowser/Source/PluginBrowser/Private/SPluginTileList.cpp
--- a/Engine/Plugins/Editor/PluginBrowser/Source/PluginBrowser/Private/SPluginTileList.cpp
+++ b/Engine/Plugins/Editor/PluginBrowser/Source/PluginBrowser/Private/SPluginTileList.cpp
@@ -1,8 +1,10 @@
 // Copyright Epic Games, Inc. All Rights Reserved.
 
 #include "SPluginTileList.h"
+#include "Framework/Views/TableViewMetadata.h"
 #include "Widgets/Views/SListView.h"
 #include "SPluginBrowser.h"
+#include "SPluginCategory.h"
 #include "SPluginTile.h"
 #include "Interfaces/IPluginManager.h"
 #include "Widgets/Layout/SScrollBorder.h"
diff --git a/Engine/Plugins/Editor/PluginBrowser/Source/PluginBrowser/Private/SPluginTileList.h b/Engine/Plugins/Editor/PluginBrowser/Source/PluginBrowser/Private/SPluginTileList.h
--- a/Engine/Plugins/Editor/PluginBrowser/Source/PluginBrowser/Private/SPluginTileList.h
+++ b/Engine/Plugins/Editor/PluginBrowser/Source/PluginBrowser/Private/SPluginTileList.h
@@ -2,11 +2,11 @@
 
 #pragma once
 
-#include "CoreMinimal.h"
-#include "Widgets/DeclarativeSyntaxSupport.h"
 #include "Widgets/SCompoundWidget.h"
-#include "Widgets/Views/STableViewBase.h"
-#include "Widgets/Views/STableRow.h"
+
+class ITableRow;
+class STableViewBase;
+template <typename ItemType> class SListView;
 
 class IPlugin;
 class SPluginBrowser;
diff --git a/Engine/Plugins/Editor/SequencerAnimTools/Source/SequencerAnimTools/Private/MotionTrailTool.cpp b/Engine/Plugins/Editor/SequencerAnimTools/Source/SequencerAnimTools/Private/MotionTrailTool.cpp
--- a/Engine/Plugins/Editor/SequencerAnimTools/Source/SequencerAnimTools/Private/MotionTrailTool.cpp
+++ b/Engine/Plugins/Editor/SequencerAnimTools/Source/SequencerAnimTools/Private/MotionTrailTool.cpp
@@ -1,30 +1,23 @@
 // Copyright Epic Games, Inc. All Rights Reserved.
 
 #include "MotionTrailTool.h"
+#include "BaseBehaviors/InputBehaviorModifierStates.h"
 #include "LevelEditorSequencerIntegration.h"
-#include "ISequencer.h"
+#include "BaseGizmos/TransformProxy.h"
 #include "Sequencer/SequencerTrailHierarchy.h"
-#include "Sequencer/MovieSceneTransformTrail.h"
 #include "BaseGizmos/CombinedTransformGizmo.h"
+#include "Framework/Commands/UICommandList.h"
 #include "InteractiveGizmoManager.h"
 #include "BaseBehaviors/SingleClickBehavior.h"
-#include "BaseBehaviors/SingleKeyCaptureBehavior.h"
-#include "BaseBehaviors/ClickDragBehavior.h"
-#include "BaseGizmos/TransformGizmoUtil.h"
 #include "Framework/Application/SlateApplication.h"
 #include "Framework/Commands/GenericCommands.h"
 
 #include "InteractiveToolManager.h"
-#include "ToolBuilderUtil.h"
 // for raycast into World
-#include "CollisionQueryParams.h"
-#include "Engine/World.h"
 
-#include "SceneManagement.h"
 #include "ScopedTransaction.h"
-#include "Tools/MotionTrailOptions.h"
 #include "EditorModeManager.h"
-#include "EditorViewportClient.h"
+#include "UnrealClient.h"
 
 #include UE_INLINE_GENERATED_CPP_BY_NAME(MotionTrailTool)
 
diff --git a/Engine/Plugins/Editor/SequencerAnimTools/Source/SequencerAnimTools/Private/Sequencer/AnimationBoneTrail.cpp b/Engine/Plugins/Editor/SequencerAnimTools/Source/SequencerAnimTools/Private/Sequencer/AnimationBoneTrail.cpp
--- a/Engine/Plugins/Editor/SequencerAnimTools/Source/SequencerAnimTools/Private/Sequencer/AnimationBoneTrail.cpp
+++ b/Engine/Plugins/Editor/SequencerAnimTools/Source/SequencerAnimTools/Private/Sequencer/AnimationBoneTrail.cpp
@@ -1,18 +1,15 @@
 // Copyright Epic Games, Inc. All Rights Reserved.
 
 #include "AnimationBoneTrail.h"
-#include "TrailHierarchy.h"
+#include "Animation/AnimData/AnimDataModel.h"
 #include "SequencerTrailHierarchy.h"
 
-#include "Animation/AnimInstance.h"
-#include "Components/SkeletalMeshComponent.h"
-#include "Sequencer/ControlRigLayerInstance.h"
+#include "Evaluation/MovieSceneSequenceTransform.h"
 #include "MovieSceneToolHelpers.h"
 
-#include "ISequencer.h"
-#include "MovieSceneSequence.h"
 #include "Exporters/AnimSeqExportOption.h"
 #include "Animation/AnimSequenceHelpers.h"
+#include "UObject/Package.h"
 
 namespace UE 
 {
diff --git a/Engine/Plugins/Editor/SequencerAnimTools/Source/SequencerAnimTools/Private/Sequencer/AnimationBoneTrail.h b/Engine/Plugins/Editor/SequencerAnimTools/Source/SequencerAnimTools/Private/Sequencer/AnimationBoneTrail.h
--- a/Engine/Plugins/Editor/SequencerAnimTools/Source/SequencerAnimTools/Private/Sequencer/AnimationBoneTrail.h
+++ b/Engine/Plugins/Editor/SequencerAnimTools/Source/SequencerAnimTools/Private/Sequencer/AnimationBoneTrail.h
@@ -2,14 +2,16 @@
 
 #pragma once
 
+#include "Animation/Skeleton.h"
 #include "Components/SkeletalMeshComponent.h"
 #include "Engine/SkeletalMesh.h"
 #include "Trail.h"
 #include "TrajectoryCache.h"
-#include "TrajectoryDrawInfo.h"
 
-#include "Sequencer/MovieSceneControlRigParameterSection.h"
 #include "Animation/AnimSequence.h"
+#include "UObject/GCObject.h"
+
+class ISequencer;
 
 namespace UE 
 {
diff --git a/Engine/Plugins/Editor/SequencerAnimTools/Source/SequencerAnimTools/Private/Sequencer/MotionTrailMovieSceneKey.cpp b/Engine/Plugins/Editor/SequencerAnimTools/Source/SequencerAnimTools/Private/Sequencer/MotionTrailMovieSceneKey.cpp
--- a/Engine/Plugins/Editor/SequencerAnimTools/Source/SequencerAnimTools/Private/Sequencer/MotionTrailMovieSceneKey.cpp
+++ b/Engine/Plugins/Editor/SequencerAnimTools/Source/SequencerAnimTools/Private/Sequencer/MotionTrailMovieSceneKey.cpp
@@ -1,19 +1,19 @@
 // Copyright Epic Games, Inc. All Rights Reserved.
 
 #include "MotionTrailMovieSceneKey.h"
+#include "Channels/MovieSceneChannelData.h"
 #include "MovieSceneTransformTrail.h"
-#include "TrailHierarchy.h"
-#include "Tracks/MovieScene3DTransformTrack.h"
-#include "Sections/MovieScene3DTransformSection.h"
 #include "Channels/MovieSceneDoubleChannel.h"
 #include "Channels/MovieSceneChannelProxy.h"
-#include "EntitySystem/Interrogation/MovieSceneInterrogationLinker.h"
-#include "HitProxies.h"
-#include "Trail.h"
+#include "Channels/MovieSceneFloatChannel.h"
 #include "Tools/MotionTrailOptions.h"
 #include "CanvasItem.h"
+#include "Engine/Engine.h"
 #include "ISequencer.h"
 #include "CanvasTypes.h"
+#include "MovieSceneSection.h"
+#include "SceneManagement.h"
+#include "SceneView.h"
 
 namespace UE
 {
diff --git a/Engine/Plugins/Editor/SequencerAnimTools/Source/SequencerAnimTools/Private/Sequencer/MotionTrailMovieSceneKey.h b/Engine/Plugins/Editor/SequencerAnimTools/Source/SequencerAnimTools/Private/Sequencer/MotionTrailMovieSceneKey.h
--- a/Engine/Plugins/Editor/SequencerAnimTools/Source/SequencerAnimTools/Private/Sequencer/MotionTrailMovieSceneKey.h
+++ b/Engine/Plugins/Editor/SequencerAnimTools/Source/SequencerAnimTools/Private/Sequencer/MotionTrailMovieSceneKey.h
@@ -2,13 +2,15 @@
 
 #pragma once
 
-#include "CoreMinimal.h"
-#include "UObject/Object.h"
+#include "Curves/KeyHandle.h"
 #include "Misc/FrameNumber.h"
-#include "MovieSceneTracksComponentTypes.h"
-#include "Algo/BinarySearch.h"
-#include "Algo/Sort.h"
-#include "Trail.h"
+
+class FCanvas;
+class FEditorViewportClient;
+class FPrimitiveDrawInterface;
+class FSceneView;
+class HHitProxy;
+namespace UE::SequencerAnimTools { struct FInputClick; }
 
 class UMovieSceneSection;
 struct FMovieSceneChannel;
@@ -256,4 +258,4 @@ protected:
 
 
 } // namespace MovieScene
-} // namespace UE
\ No newline at end of file
+} // namespace UE
diff --git a/Engine/Plugins/Editor/SequencerAnimTools/Source/SequencerAnimTools/Private/Sequencer/MovieSceneTransformTrail.cpp b/Engine/Plugins/Editor/SequencerAnimTools/Source/SequencerAnimTools/Private/Sequencer/MovieSceneTransformTrail.cpp
--- a/Engine/Plugins/Editor/SequencerAnimTools/Source/SequencerAnimTools/Private/Sequencer/MovieSceneTransformTrail.cpp
+++ b/Engine/Plugins/Editor/SequencerAnimTools/Source/SequencerAnimTools/Private/Sequencer/MovieSceneTransformTrail.cpp
@@ -1,32 +1,19 @@
 // Copyright Epic Games, Inc. All Rights Reserved.
 
 #include "MovieSceneTransformTrail.h"
+#include "Engine/Engine.h"
 #include "TrailHierarchy.h"
 
-#include "ISequencer.h"
-#include "Tracks/MovieScene3DTransformTrack.h"
-#include "Sections/MovieScene3DTransformSection.h"
-#include "Channels/MovieSceneDoubleChannel.h"
 #include "Channels/MovieSceneChannelProxy.h"
-#include "Systems/MovieSceneComponentTransformSystem.h"
-#include "MovieSceneSequence.h"
-#include "MovieSceneTracksComponentTypes.h"
+#include "Evaluation/MovieSceneEvaluationTemplateInstance.h"
 #include "Sequencer/MovieSceneControlRigParameterSection.h"
 
-#include "BaseGizmos/GizmoComponents.h"
-#include "BaseGizmos/CombinedTransformGizmo.h"
 
-#include "ViewportWorldInteraction.h"
-#include "GameFramework/Actor.h"
-#include "Components/SceneComponent.h"
 #include "MovieSceneToolHelpers.h"
 #include "ActorForWorldTransforms.h"
-#include "IControlRigObjectBinding.h"
-#include "LevelSequence.h"
-#include "LevelSequenceActor.h"
 #include "Evaluation/MovieScenePlayback.h"
-#include "Tools/ControlRigSnapSettings.h"
 #include "ScopedTransaction.h"
+#include "Sequencer/MotionTrailMovieSceneKey.h"
 
 #define LOCTEXT_NAMESPACE "SequencerAnimTools"
 
diff --git a/Engine/Plugins/Editor/SequencerAnimTools/Source/SequencerAnimTools/Private/Sequencer/MovieSceneTransformTrail.h b/Engine/Plugins/Editor/SequencerAnimTools/Source/SequencerAnimTools/Private/Sequencer/MovieSceneTransformTrail.h
--- a/Engine/Plugins/Editor/SequencerAnimTools/Source/SequencerAnimTools/Private/Sequencer/MovieSceneTransformTrail.h
+++ b/Engine/Plugins/Editor/SequencerAnimTools/Source/SequencerAnimTools/Private/Sequencer/MovieSceneTransformTrail.h
@@ -2,10 +2,13 @@
 
 #pragma once
 
+#include "Components/SceneComponent.h"
 #include "Trail.h"
 #include "MotionTrailMovieSceneKey.h"
-#include "TrajectoryDrawInfo.h"
-#include "EntitySystem/Interrogation/MovieSceneInterrogationLinker.h"
+#include "MovieSceneTrack.h"
+#include "TrajectoryCache.h"
+
+class IMovieScenePlayer;
 
 class ISequencer;
 class UMovieScene3DTransformTrack;
diff --git a/Engine/Plugins/Editor/SequencerAnimTools/Source/SequencerAnimTools/Private/Sequencer/SequencerTrailHierarchy.cpp b/Engine/Plugins/Editor/SequencerAnimTools/Source/SequencerAnimTools/Private/Sequencer/SequencerTrailHierarchy.cpp
--- a/Engine/Plugins/Editor/SequencerAnimTools/Source/SequencerAnimTools/Private/Sequencer/SequencerTrailHierarchy.cpp
+++ b/Engine/Plugins/Editor/SequencerAnimTools/Source/SequencerAnimTools/Private/Sequencer/SequencerTrailHierarchy.cpp
@@ -2,28 +2,24 @@
 
 #include "SequencerTrailHierarchy.h"
 
-#include "Components/SkeletalMeshComponent.h"
-#include "MovieSceneSequence.h"
-#include "ISequencer.h"
+#include "Editor.h"
 #include "Tracks/MovieScene3DTransformTrack.h"
+#include "Evaluation/MovieSceneEvaluationTemplateInstance.h"
 #include "Tracks/MovieSceneSkeletalAnimationTrack.h"
-#include "MovieSceneSection.h"
 
-#include "ControlRig.h"
+#include "Evaluation/MovieScenePlayback.h"
 #include "Sequencer/MovieSceneControlRigParameterTrack.h"
-#include "Sequencer/MovieSceneControlRigParameterSection.h"
 #include "AnimationBoneTrail.h"
-#include "LevelSequencePlayer.h"
+#include "LevelSequence.h"
 #include "Tools/MotionTrailOptions.h"
+#include "LevelSequenceCameraSettings.h"
 #include "MovieSceneTransformTrail.h"
 #include "IControlRigObjectBinding.h"
 #include "ActorForWorldTransforms.h"
+#include "MovieSceneSequencePlaybackSettings.h"
 #include "MovieSceneToolHelpers.h"
-#include "ControlRigSequencerEditorLibrary.h"
-#include "Engine/World.h"
-#include "Kismet/GameplayStatics.h"
-#include "EngineUtils.h"
 #include "Framework/Application/SlateApplication.h"
+#include "Sequencer/MotionTrailMovieSceneKey.h"
 
 namespace UE
 {
@@ -228,7 +224,6 @@ void FSequencerTrailHierarchy::UpdateControlRig(const TArray<FFrameNumber> &Fram
 		FFrameRate TickResolution = Sequencer->GetFocusedTickResolution();
 
 		UWorld* World = GEditor->GetEditorWorldContext().World();
-		ALevelSequenceActor* OutActor = nullptr;
 		FMovieSceneSequencePlaybackSettings Settings;
 		FLevelSequenceCameraSettings CameraSettings;
 		FMovieSceneSequenceIDRef Template = MovieSceneSequenceID::Root;
diff --git a/Engine/Plugins/Editor/SequencerAnimTools/Source/SequencerAnimTools/Private/Sequencer/SequencerTrailHierarchy.h b/Engine/Plugins/Editor/SequencerAnimTools/Source/SequencerAnimTools/Private/Sequencer/SequencerTrailHierarchy.h
--- a/Engine/Plugins/Editor/SequencerAnimTools/Source/SequencerAnimTools/Private/Sequencer/SequencerTrailHierarchy.h
+++ b/Engine/Plugins/Editor/SequencerAnimTools/Source/SequencerAnimTools/Private/Sequencer/SequencerTrailHierarchy.h
@@ -4,8 +4,9 @@
 
 #include "TrailHierarchy.h"
 #include "ISequencer.h"
-#include "Rigs/RigHierarchyDefines.h"
-#include "UObject/GCObject.h"
+
+namespace UE::SequencerAnimTools { class FTrail; }
+namespace UE::SequencerAnimTools { class FTrailEvaluateTimes; }
 
 
 class ISequencer;
diff --git a/Engine/Plugins/Editor/SequencerAnimTools/Source/SequencerAnimTools/Private/SequencerAnimEditPivotTool.cpp b/Engine/Plugins/Editor/SequencerAnimTools/Source/SequencerAnimTools/Private/SequencerAnimEditPivotTool.cpp
--- a/Engine/Plugins/Editor/SequencerAnimTools/Source/SequencerAnimTools/Private/SequencerAnimEditPivotTool.cpp
+++ b/Engine/Plugins/Editor/SequencerAnimTools/Source/SequencerAnimTools/Private/SequencerAnimEditPivotTool.cpp
@@ -1,30 +1,21 @@
 // Copyright Epic Games, Inc. All Rights Reserved.
 
 #include "SequencerAnimEditPivotTool.h"
-#include "LevelEditorSequencerIntegration.h"
+#include "BaseGizmos/TransformProxy.h"
 #include "ISequencer.h"
-#include "Sequencer/SequencerTrailHierarchy.h"
-#include "Sequencer/MovieSceneTransformTrail.h"
 #include "BaseGizmos/CombinedTransformGizmo.h"
+#include "IControlRigObjectBinding.h"
 #include "InteractiveGizmoManager.h"
 #include "BaseBehaviors/SingleClickBehavior.h"
-#include "BaseBehaviors/SingleKeyCaptureBehavior.h"
-#include "BaseGizmos/TransformGizmoUtil.h"
 #include "Framework/Application/SlateApplication.h"
-#include "Framework/Commands/GenericCommands.h"
+#include "ILevelEditor.h"
 #include "LevelEditor.h"
 
 #include "InteractiveToolManager.h"
-#include "ToolBuilderUtil.h"
 // for raycast into World
-#include "CollisionQueryParams.h"
-#include "Engine/World.h"
 
-#include "SceneManagement.h"
-#include "ScopedTransaction.h"
-#include "Tools/MotionTrailOptions.h"
 #include "EditorModeManager.h"
-#include "EditorViewportClient.h"
+#include "Misc/QualifiedFrameTime.h"
 #include "Modules/ModuleManager.h"
 #include "ControlRig.h"
 #include "ControlRigSequencerEditorLibrary.h"
@@ -33,23 +24,16 @@
 #include "ILevelSequenceEditorToolkit.h"
 #include "InteractiveToolManager.h"
 #include "EdModeInteractiveToolsContext.h"
+#include "SceneView.h"
 #include "UnrealEdGlobals.h"
 #include "Editor/UnrealEdEngine.h"
-#include "Subsystems/EditorActorSubsystem.h"
-#include "LevelEditorViewport.h"
 #include "Modules/ModuleManager.h"
-#include "IAssetViewport.h"
-#include "Widgets/SWindow.h"
-#include "Widgets/Layout/SBorder.h"
+#include "Selection.h"
 #include "Widgets/Text/STextBlock.h"
-#include "Widgets/Layout/SBox.h"
 #include "Widgets/Input/SButton.h"
-#include "Styling/AppStyle.h"
 #include "Viewports/InViewportUIDragOperation.h"
 #include "Widgets/Layout/SUniformGridPanel.h"
 #include "Styling/StyleColors.h"
-#include "EditorModes.h"
-#include "ControlRigObjectBinding.h"
 #include "SLevelViewport.h"
 
 #define LOCTEXT_NAMESPACE "SequencerAnimTools"
@@ -1006,7 +990,6 @@ void USequencerPivotTool::OnPropertyModified(UObject* PropertySet, FProperty* Pr
 	//UpdateDistance();
 }
 
-#include "HitProxies.h"
 
 #include UE_INLINE_GENERATED_CPP_BY_NAME(SequencerAnimEditPivotTool)
 
diff --git a/Engine/Plugins/Editor/SequencerAnimTools/Source/SequencerAnimTools/Private/SequencerAnimToolsModule.cpp b/Engine/Plugins/Editor/SequencerAnimTools/Source/SequencerAnimTools/Private/SequencerAnimToolsModule.cpp
--- a/Engine/Plugins/Editor/SequencerAnimTools/Source/SequencerAnimTools/Private/SequencerAnimToolsModule.cpp
+++ b/Engine/Plugins/Editor/SequencerAnimTools/Source/SequencerAnimTools/Private/SequencerAnimToolsModule.cpp
@@ -2,13 +2,13 @@
 
 #include "SequencerAnimToolsModule.h"
 #include "EditorModeManager.h"
-#include "Editor.h"
-#include "EditorModeRegistry.h"
+#include "ILevelEditor.h"
 #include "LevelEditor.h"
 #include "EdModeInteractiveToolsContext.h"
-#include "UnrealEdGlobals.h"
+#include "InteractiveToolManager.h"
 #include "LevelEditorSequencerIntegration.h"
 
+#include "Modules/ModuleManager.h"
 #include "MotionTrailTool.h"
 #include "Tools/MotionTrailOptions.h"
 
@@ -84,4 +84,4 @@ void FSequencerAnimToolsModule::OnMotionTralOptionChanged(FName PropertyName)
 
 #undef LOCTEXT_NAMESPACE
 	
-IMPLEMENT_MODULE(UE::SequencerAnimTools::FSequencerAnimToolsModule, SequencerAnimTools)
\ No newline at end of file
+IMPLEMENT_MODULE(UE::SequencerAnimTools::FSequencerAnimToolsModule, SequencerAnimTools)
diff --git a/Engine/Plugins/Editor/SequencerAnimTools/Source/SequencerAnimTools/Private/SequencerToolsEditMode.cpp b/Engine/Plugins/Editor/SequencerAnimTools/Source/SequencerAnimTools/Private/SequencerToolsEditMode.cpp
--- a/Engine/Plugins/Editor/SequencerAnimTools/Source/SequencerAnimTools/Private/SequencerToolsEditMode.cpp
+++ b/Engine/Plugins/Editor/SequencerAnimTools/Source/SequencerAnimTools/Private/SequencerToolsEditMode.cpp
@@ -3,12 +3,11 @@
 #include "SequencerToolsEditMode.h"
 #include "EditorModeManager.h"
 #include "EdModeInteractiveToolsContext.h"
-#include "EdMode.h"
+#include "ILevelEditor.h"
 #include "LevelEditor.h"
 #include "Modules/ModuleManager.h"
 #include "EditorModeManager.h"
 #include "InteractiveToolManager.h"
-#include "BaseSequencerAnimTool.h"
 #include "BaseGizmos/TransformGizmoUtil.h"
 #include "MotionTrailTool.h"
 #include "SequencerAnimEditPivotTool.h"
diff --git a/Engine/Plugins/Editor/SequencerAnimTools/Source/SequencerAnimTools/Private/Trail.cpp b/Engine/Plugins/Editor/SequencerAnimTools/Source/SequencerAnimTools/Private/Trail.cpp
--- a/Engine/Plugins/Editor/SequencerAnimTools/Source/SequencerAnimTools/Private/Trail.cpp
+++ b/Engine/Plugins/Editor/SequencerAnimTools/Source/SequencerAnimTools/Private/Trail.cpp
@@ -1,7 +1,6 @@
 // Copyright Epic Games, Inc. All Rights Reserved.
 
 #include "Trail.h"
-#include "TrailHierarchy.h"
 
 //MZ TODO LOOK AT REIMPLEMENTING CONSTANT TRAILS
 
diff --git a/Engine/Plugins/Editor/SequencerAnimTools/Source/SequencerAnimTools/Private/TrailHierarchy.cpp b/Engine/Plugins/Editor/SequencerAnimTools/Source/SequencerAnimTools/Private/TrailHierarchy.cpp
--- a/Engine/Plugins/Editor/SequencerAnimTools/Source/SequencerAnimTools/Private/TrailHierarchy.cpp
+++ b/Engine/Plugins/Editor/SequencerAnimTools/Source/SequencerAnimTools/Private/TrailHierarchy.cpp
@@ -1,16 +1,18 @@
 // Copyright Epic Games, Inc. All Rights Reserved.
 
 #include "TrailHierarchy.h"
-#include "TrajectoryDrawInfo.h"
-#include "Containers/Queue.h"
 #include "CanvasItem.h"
 #include "CanvasTypes.h"
+#include "ILevelEditor.h"
 #include "LevelEditor.h"
+#include "Layout/WidgetPath.h"
 #include "Modules/ModuleManager.h"
 #include "Framework/Application/SlateApplication.h"
-#include "EditorModeManager.h"
+#include "SceneManagement.h"
 #include "Tools/MotionTrailOptions.h"
+#include "TrajectoryCache.h"
 #include "Widgets/Input/SButton.h"
+#include "Widgets/Input/SCheckBox.h"
 #include "Widgets/Text/STextBlock.h"
 #include "CanvasTypes.h"
 
