ModelingMode: add a new config setting to enable/disable Mesh Selection UI, default to true, and hide old setting, so that it defaults to on for everyone

diff --git a/Engine/Plugins/Editor/ModelingToolsEditorMode/Source/ModelingToolsEditorMode/Private/ModelingToolsEditorMode.cpp b/Engine/Plugins/Editor/ModelingToolsEditorMode/Source/ModelingToolsEditorMode/Private/ModelingToolsEditorMode.cpp
--- a/Engine/Plugins/Editor/ModelingToolsEditorMode/Source/ModelingToolsEditorMode/Private/ModelingToolsEditorMode.cpp
+++ b/Engine/Plugins/Editor/ModelingToolsEditorMode/Source/ModelingToolsEditorMode/Private/ModelingToolsEditorMode.cpp
@@ -379,7 +379,7 @@ void UModelingToolsEditorMode::Enter()
 	SceneSnappingManager = UE::Geometry::FindModelingSceneSnappingManager(GetToolManager());
 
 	// register selection manager, if this feature is enabled in the mode settings
-	if (ModelingModeSettings && ModelingModeSettings->bEnablePersistentSelections)
+	if (ModelingModeSettings && ModelingModeSettings->bEnableMeshSelections)
 	{
 		// set up SelectionManager and register known factory types
 		SelectionManager = NewObject<UGeometrySelectionManager>(GetToolManager());
diff --git a/Engine/Plugins/Editor/ModelingToolsEditorMode/Source/ModelingToolsEditorMode/Private/ModelingToolsEditorModeToolkit.cpp b/Engine/Plugins/Editor/ModelingToolsEditorMode/Source/ModelingToolsEditorMode/Private/ModelingToolsEditorModeToolkit.cpp
--- a/Engine/Plugins/Editor/ModelingToolsEditorMode/Source/ModelingToolsEditorMode/Private/ModelingToolsEditorModeToolkit.cpp
+++ b/Engine/Plugins/Editor/ModelingToolsEditorMode/Source/ModelingToolsEditorMode/Private/ModelingToolsEditorModeToolkit.cpp
@@ -172,7 +172,7 @@ void FModelingToolsEditorModeToolkit::Init(const TSharedPtr<IToolkitHost>& InitT
 	MakeToolShutdownOverlayWidget();
 
 	const UModelingToolsEditorModeSettings* ModelingModeSettings = GetDefault<UModelingToolsEditorModeSettings>();
-	bool bEnableSelectionUI = ModelingModeSettings && ModelingModeSettings->bEnablePersistentSelections;
+	bool bEnableSelectionUI = ModelingModeSettings && ModelingModeSettings->bEnableMeshSelections;
 	if ( bEnableSelectionUI )
 	{
 		MakeSelectionPaletteOverlayWidget();
@@ -567,7 +567,7 @@ void FModelingToolsEditorModeToolkit::GetToolPaletteNames(TArray<FName>& Palette
 	}
 
 	const UModelingToolsEditorModeSettings* ModelingModeSettings = GetDefault<UModelingToolsEditorModeSettings>();
-	bool bEnableSelectionUI = ModelingModeSettings && ModelingModeSettings->bEnablePersistentSelections;
+	bool bEnableSelectionUI = ModelingModeSettings && ModelingModeSettings->bEnableMeshSelections;
 	if (bEnableSelectionUI)
 	{
 		if (bShowActiveSelectionActions)
diff --git a/Engine/Plugins/Editor/ModelingToolsEditorMode/Source/ModelingToolsEditorMode/Public/ModelingToolsEditorModeSettings.h b/Engine/Plugins/Editor/ModelingToolsEditorMode/Source/ModelingToolsEditorMode/Public/ModelingToolsEditorModeSettings.h
--- a/Engine/Plugins/Editor/ModelingToolsEditorMode/Source/ModelingToolsEditorMode/Public/ModelingToolsEditorModeSettings.h
+++ b/Engine/Plugins/Editor/ModelingToolsEditorMode/Source/ModelingToolsEditorMode/Public/ModelingToolsEditorModeSettings.h
@@ -145,10 +145,15 @@ public:
 	// Selection
 	//
 
-	/** If true, Mesh Selection System (Experimental) will be enabled. */
-	UPROPERTY(config, EditAnywhere, Category = "Modeling Mode|Selection", meta = (DisplayName="Enable Mesh Selection UI"))
+	// old preference that should be ignored now
+	UPROPERTY()
 	bool bEnablePersistentSelections = false;
 
+	/** Enable/Disable new Mesh Selection System (Experimental). */
+	UPROPERTY(config, EditAnywhere, Category = "Modeling Mode|Selection", meta = (DisplayName="Enable Mesh Selection UI"))
+	bool bEnableMeshSelections = true;
+
+
 	DECLARE_MULTICAST_DELEGATE_TwoParams(UModelingToolsEditorModeSettingsModified, UObject*, FProperty*);
 	UModelingToolsEditorModeSettingsModified OnModified;
 
