ModelingMode: add support for Editor Marquee in new Mesh Selection system.

Add IGeometrySelector::UpdateSelectionViaShape() API, takes a FConvexHull to use as selection shape. Implement in DynamicMeshSelector.
Add UGeometrySelectionManager::UpdateSelectionViaConvex() API, calls above
ModelingToolsEditorMode now implements ILegacyEdModeSelectInterface, ::FrustumSelect() implementation calls above w/ frustum convex hull

Hide Marquee drag mode option, set default drag mode to "none"
Add save/restore of selection mode, drag mode, and volume/staticmesh toggles via UModelingToolsModeCustomizationSettings
Remove old selection toggle commands that were replaced w/ combined commands
Call DeleteSelectionCommand in UModelingToolsEditorMode::ProcessEditDelete() if there is an active selection

diff --git a/Engine/Plugins/Editor/ModelingToolsEditorMode/Source/ModelingToolsEditorMode/Private/ModelingModeToolkit_Toolbars.cpp b/Engine/Plugins/Editor/ModelingToolsEditorMode/Source/ModelingToolsEditorMode/Private/ModelingModeToolkit_Toolbars.cpp
--- a/Engine/Plugins/Editor/ModelingToolsEditorMode/Source/ModelingToolsEditorMode/Private/ModelingModeToolkit_Toolbars.cpp
+++ b/Engine/Plugins/Editor/ModelingToolsEditorMode/Source/ModelingToolsEditorMode/Private/ModelingModeToolkit_Toolbars.cpp
@@ -301,6 +301,9 @@ void MakeSubMenu_Selection_DragMode(FModelingToolsEditorModeToolkit* Toolkit, FM
 			FExecuteAction::CreateLambda([SelectionInteraction, DragMode]
 			{
 				SelectionInteraction->SetActiveDragMode(DragMode);
+				UModelingToolsModeCustomizationSettings* ModelingEditorSettings = GetMutableDefault<UModelingToolsModeCustomizationSettings>();
+				ModelingEditorSettings->LastMeshSelectionDragMode = static_cast<int>(DragMode);
+				ModelingEditorSettings->SaveConfig();
 			}),
 			FCanExecuteAction(),
 			FIsActionChecked::CreateLambda([SelectionInteraction, DragMode]()
@@ -319,10 +322,11 @@ void MakeSubMenu_Selection_DragMode(FModelingToolsEditorModeToolkit* Toolkit, FM
 		LOCTEXT("Selection_DragInput_Path_Tooltip", "Path Drag Input"),
 		FSlateIcon(), MakeDragModeOptionAction(EModelingSelectionInteraction_DragMode::PathInteraction), NAME_None, EUserInterfaceActionType::ToggleButton);
 
-	MenuBuilder.AddMenuEntry(
-		LOCTEXT("Selection_DragInput_Marquee", "Rectangle"), 
-		LOCTEXT("Selection_DragInput_Marquee_Tooltip", "Rectangle Marquee"),
-		FSlateIcon(), MakeDragModeOptionAction(EModelingSelectionInteraction_DragMode::RectangleMarqueeInteraction), NAME_None, EUserInterfaceActionType::ToggleButton);
+	// marquee mode is not functional yet, so this is disabled
+	//MenuBuilder.AddMenuEntry(
+	//	LOCTEXT("Selection_DragInput_Marquee", "Rectangle"), 
+	//	LOCTEXT("Selection_DragInput_Marquee_Tooltip", "Rectangle Marquee"),
+	//	FSlateIcon(), MakeDragModeOptionAction(EModelingSelectionInteraction_DragMode::RectangleMarqueeInteraction), NAME_None, EUserInterfaceActionType::ToggleButton);
 
 }
 
@@ -334,6 +338,9 @@ void MakeSubMenu_Selection_MeshType(FModelingToolsEditorModeToolkit* Toolkit, FM
 		FExecuteAction::CreateLambda([ModelingMode]
 		{
 			ModelingMode->bEnableVolumeElementSelection = ! ModelingMode->bEnableVolumeElementSelection;
+
+			UModelingToolsModeCustomizationSettings* ModelingEditorSettings = GetMutableDefault<UModelingToolsModeCustomizationSettings>();
+			ModelingEditorSettings->bLastMeshSelectionVolumeToggle = ModelingMode->bEnableVolumeElementSelection;
 		}),
 		FCanExecuteAction(),
 		FIsActionChecked::CreateLambda([ModelingMode]()
@@ -344,6 +351,9 @@ void MakeSubMenu_Selection_MeshType(FModelingToolsEditorModeToolkit* Toolkit, FM
 		FExecuteAction::CreateLambda([ModelingMode]
 		{
 			ModelingMode->bEnableStaticMeshElementSelection = ! ModelingMode->bEnableStaticMeshElementSelection;
+
+			UModelingToolsModeCustomizationSettings* ModelingEditorSettings = GetMutableDefault<UModelingToolsModeCustomizationSettings>();
+			ModelingEditorSettings->bLastMeshSelectionStaticMeshToggle = ModelingMode->bEnableStaticMeshElementSelection;
 		}),
 		FCanExecuteAction(),
 		FIsActionChecked::CreateLambda([ModelingMode]()
diff --git a/Engine/Plugins/Editor/ModelingToolsEditorMode/Source/ModelingToolsEditorMode/Private/ModelingToolsEditorMode.cpp b/Engine/Plugins/Editor/ModelingToolsEditorMode/Source/ModelingToolsEditorMode/Private/ModelingToolsEditorMode.cpp
--- a/Engine/Plugins/Editor/ModelingToolsEditorMode/Source/ModelingToolsEditorMode/Private/ModelingToolsEditorMode.cpp
+++ b/Engine/Plugins/Editor/ModelingToolsEditorMode/Source/ModelingToolsEditorMode/Private/ModelingToolsEditorMode.cpp
@@ -177,6 +177,15 @@ UModelingToolsEditorMode::~UModelingToolsEditorMode()
 
 bool UModelingToolsEditorMode::ProcessEditDelete()
 {
+	if (SelectionManager && SelectionManager->HasSelection())
+	{
+		if ( UDeleteGeometrySelectionCommand* DeleteCommand = Cast<UDeleteGeometrySelectionCommand>(ModelingModeCommands[0]) )
+		{
+			SelectionManager->ExecuteSelectionCommand(DeleteCommand);
+			return true;
+		}
+	}
+
 	if (UEdMode::ProcessEditDelete())
 	{
 		return true;
@@ -751,36 +760,6 @@ void UModelingToolsEditorMode::Enter()
 	
 
 	// set up selection type toggles
-	auto RegisterSelectionTopologyType = [this](UGeometrySelectionManager::EMeshTopologyMode TopoMode, TSharedPtr<FUICommandInfo> UICommand)
-	{
-		Toolkit->GetToolkitCommands()->MapAction(UICommand,
-			FExecuteAction::CreateLambda([this, TopoMode]() { GetSelectionManager()->SetMeshTopologyMode(TopoMode); }),
-			FCanExecuteAction(),
-			FIsActionChecked::CreateLambda([this, TopoMode]() { return GetSelectionManager()->GetMeshTopologyMode() == TopoMode; }),
-			EUIActionRepeatMode::RepeatDisabled);
-	};
-	if (GetSelectionManager() != nullptr)
-	{
-		RegisterSelectionTopologyType(UGeometrySelectionManager::EMeshTopologyMode::None, ToolManagerCommands.BeginSelectionAction_ToObjectType);
-		RegisterSelectionTopologyType(UGeometrySelectionManager::EMeshTopologyMode::Triangle, ToolManagerCommands.BeginSelectionAction_ToTriangleType);
-		RegisterSelectionTopologyType(UGeometrySelectionManager::EMeshTopologyMode::Polygroup, ToolManagerCommands.BeginSelectionAction_ToPolygroupType);
-	}
-
-	auto RegisterSelectionElementType = [this](UE::Geometry::EGeometryElementType ElementMode, TSharedPtr<FUICommandInfo> UICommand)
-	{
-		Toolkit->GetToolkitCommands()->MapAction(UICommand,
-			FExecuteAction::CreateLambda([this, ElementMode]() { GetSelectionManager()->SetSelectionElementType(ElementMode); }),
-			FCanExecuteAction(),
-			FIsActionChecked::CreateLambda([this, ElementMode]() { return GetSelectionManager()->GetSelectionElementType() == ElementMode; }),
-			EUIActionRepeatMode::RepeatDisabled);
-	};
-	if (GetSelectionManager() != nullptr)
-	{
-		RegisterSelectionElementType(UE::Geometry::EGeometryElementType::Vertex, ToolManagerCommands.BeginSelectionAction_ToVertexType);
-		RegisterSelectionElementType(UE::Geometry::EGeometryElementType::Edge, ToolManagerCommands.BeginSelectionAction_ToEdgeType);
-		RegisterSelectionElementType(UE::Geometry::EGeometryElementType::Face, ToolManagerCommands.BeginSelectionAction_ToFaceType);
-	}
-
 
 	auto RegisterSelectionMode = [this](UGeometrySelectionManager::EMeshTopologyMode TopoMode, UE::Geometry::EGeometryElementType ElementMode, TSharedPtr<FUICommandInfo> UICommand)
 	{
@@ -790,6 +769,11 @@ void UModelingToolsEditorMode::Enter()
 				GetSelectionManager()->SetMeshTopologyMode(TopoMode); 
 				GetSelectionManager()->SetSelectionElementType(ElementMode);
 				GetToolManager()->GetContextTransactionsAPI()->EndUndoTransaction();
+
+				UModelingToolsModeCustomizationSettings* ModelingEditorSettings = GetMutableDefault<UModelingToolsModeCustomizationSettings>();
+				ModelingEditorSettings->LastMeshSelectionTopologyMode = static_cast<int>(TopoMode);
+				ModelingEditorSettings->LastMeshSelectionElementType = static_cast<int>(ElementMode);
+				ModelingEditorSettings->SaveConfig();
 			}),
 			FCanExecuteAction(),
 			FIsActionChecked::CreateLambda([this, TopoMode, ElementMode]() { return GetSelectionManager()->GetMeshTopologyMode() == TopoMode && GetSelectionManager()->GetSelectionElementType() == ElementMode; }),
@@ -930,9 +914,36 @@ void UModelingToolsEditorMode::Enter()
 	SelectionModifiedEventHandle = GetModeManager()->GetSelectedActors()->SelectionChangedEvent.AddLambda(
 		[this](const UObject* Object)  { UpdateSelectionManagerOnEditorSelectionChange(); } );
 
+	// restore various settings
+	const UModelingToolsModeCustomizationSettings* ModelingEditorSettings = GetDefault<UModelingToolsModeCustomizationSettings>();
+	if (SelectionManager)
+	{
+		UE::Geometry::EGeometryElementType LastElementType = static_cast<UE::Geometry::EGeometryElementType>(ModelingEditorSettings->LastMeshSelectionElementType);
+		if (LastElementType == UE::Geometry::EGeometryElementType::Edge || LastElementType == UE::Geometry::EGeometryElementType::Face || LastElementType == UE::Geometry::EGeometryElementType::Vertex)
+		{
+			SelectionManager->SetSelectionElementType(LastElementType);
+		}
+		UGeometrySelectionManager::EMeshTopologyMode LastTopologyMode = static_cast<UGeometrySelectionManager::EMeshTopologyMode>(ModelingEditorSettings->LastMeshSelectionTopologyMode);
+		if (LastTopologyMode == UGeometrySelectionManager::EMeshTopologyMode::None || LastTopologyMode == UGeometrySelectionManager::EMeshTopologyMode::Triangle || LastTopologyMode == UGeometrySelectionManager::EMeshTopologyMode::Polygroup)
+		{
+			SelectionManager->SetMeshTopologyMode(LastTopologyMode);
+		}
+
+		bEnableStaticMeshElementSelection = ModelingEditorSettings->bLastMeshSelectionStaticMeshToggle;
+		bEnableVolumeElementSelection = ModelingEditorSettings->bLastMeshSelectionVolumeToggle;
+	}
+	if ( SelectionInteraction )
+	{
+		EModelingSelectionInteraction_DragMode LastDragMode = static_cast<EModelingSelectionInteraction_DragMode>(ModelingEditorSettings->LastMeshSelectionDragMode);
+		if ( LastDragMode == EModelingSelectionInteraction_DragMode::NoDragInteraction || LastDragMode == EModelingSelectionInteraction_DragMode::PathInteraction )
+		{
+			SelectionInteraction->SetActiveDragMode(LastDragMode);
+		}
+	}
+
+
 	// initialize SelectionManager w/ active selection
 	UpdateSelectionManagerOnEditorSelectionChange(true);
-
 }
 
 void UModelingToolsEditorMode::RegisterUVEditor()
@@ -1227,6 +1238,40 @@ void UModelingToolsEditorMode::UpdateSelectionManagerOnEditorSelectionChange(boo
 }
 
 
+bool UModelingToolsEditorMode::BoxSelect(FBox& InBox, bool InSelect) 
+{
+	// not handling yet
+	return false;
+}
+
+bool UModelingToolsEditorMode::FrustumSelect(const FConvexVolume& InFrustum, FEditorViewportClient* InViewportClient, bool InSelect)
+{
+	if (SelectionManager && SelectionManager->HasActiveTargets())
+	{
+		UE::Geometry::FGeometrySelectionUpdateConfig UpdateConfig;
+		UpdateConfig.ChangeType = UE::Geometry::EGeometrySelectionChangeType::Replace;
+		if ( InViewportClient->IsShiftPressed() )
+		{
+			UpdateConfig.ChangeType = UE::Geometry::EGeometrySelectionChangeType::Add;
+		}
+		else if ( InViewportClient->IsCtrlPressed() && InViewportClient->IsAltPressed() == false )
+		{
+			UpdateConfig.ChangeType = UE::Geometry::EGeometrySelectionChangeType::Remove;
+		}			
+
+		UE::Geometry::FGeometrySelectionUpdateResult Result;
+		SelectionManager->UpdateSelectionViaConvex(
+			InFrustum, UpdateConfig, Result);
+
+		return true;	// always consume marquee even if it missed, as the miss will usually just be a mistake
+	}
+
+	// not handling yet
+	return false;
+}
+
+
+
 void UModelingToolsEditorMode::CreateToolkit()
 {
 	Toolkit = MakeShareable(new FModelingToolsEditorModeToolkit);
diff --git a/Engine/Plugins/Editor/ModelingToolsEditorMode/Source/ModelingToolsEditorMode/Private/ModelingToolsManagerActions.cpp b/Engine/Plugins/Editor/ModelingToolsEditorMode/Source/ModelingToolsEditorMode/Private/ModelingToolsManagerActions.cpp
--- a/Engine/Plugins/Editor/ModelingToolsEditorMode/Source/ModelingToolsEditorMode/Private/ModelingToolsManagerActions.cpp
+++ b/Engine/Plugins/Editor/ModelingToolsEditorMode/Source/ModelingToolsEditorMode/Private/ModelingToolsManagerActions.cpp
@@ -182,14 +182,6 @@ void FModelingToolsManagerCommands::RegisterCommands()
 	REGISTER_MODELING_TOOL_COMMAND(BeginPolyModelTool_PushPull, "PushPull", "PolyGroup Push Pull Faces Tool");
 	REGISTER_MODELING_TOOL_COMMAND(BeginPolyModelTool_Bevel, "Bevel", "PolyGroup Bevel Tool");
 
-
-	REGISTER_MODELING_TOOL_COMMAND(BeginSelectionAction_ToObjectType, "None", "Disable Mesh Element Selection");
-	REGISTER_MODELING_TOOL_COMMAND(BeginSelectionAction_ToTriangleType, "Tris", "Set Selection Layer to Triangles");
-	REGISTER_MODELING_TOOL_COMMAND(BeginSelectionAction_ToPolygroupType, "Groups", "Set Selection Layer to Polygroups");
-	REGISTER_MODELING_TOOL_COMMAND(BeginSelectionAction_ToVertexType, "Vertex", "Set Selection Type to Vertices");
-	REGISTER_MODELING_TOOL_COMMAND(BeginSelectionAction_ToEdgeType, "Edge", "Set Selection Type to Edges");
-	REGISTER_MODELING_TOOL_COMMAND(BeginSelectionAction_ToFaceType, "Face", "Set Selection Type to Faces");
-
 	REGISTER_MODELING_TOOL_COMMAND(MeshSelectionModeAction_NoSelection, "None", "Disable Geometry Selection");
 	REGISTER_MODELING_TOOL_COMMAND(MeshSelectionModeAction_MeshTriangles, "Tris", "Select Mesh Triangles");
 	REGISTER_MODELING_TOOL_COMMAND(MeshSelectionModeAction_MeshVertices, "Verts", "Select Mesh Vertices");
diff --git a/Engine/Plugins/Editor/ModelingToolsEditorMode/Source/ModelingToolsEditorMode/Public/ModelingSelectionInteraction.h b/Engine/Plugins/Editor/ModelingToolsEditorMode/Source/ModelingToolsEditorMode/Public/ModelingSelectionInteraction.h
--- a/Engine/Plugins/Editor/ModelingToolsEditorMode/Source/ModelingToolsEditorMode/Public/ModelingSelectionInteraction.h
+++ b/Engine/Plugins/Editor/ModelingToolsEditorMode/Source/ModelingToolsEditorMode/Public/ModelingSelectionInteraction.h
@@ -201,7 +201,7 @@ protected:
 
 
 protected:
-	EModelingSelectionInteraction_DragMode ActiveDragMode = EModelingSelectionInteraction_DragMode::PathInteraction;
+	EModelingSelectionInteraction_DragMode ActiveDragMode = EModelingSelectionInteraction_DragMode::NoDragInteraction;
 	void UpdateActiveDragMode();
 
 	UPROPERTY()
diff --git a/Engine/Plugins/Editor/ModelingToolsEditorMode/Source/ModelingToolsEditorMode/Public/ModelingToolsEditorMode.h b/Engine/Plugins/Editor/ModelingToolsEditorMode/Source/ModelingToolsEditorMode/Public/ModelingToolsEditorMode.h
--- a/Engine/Plugins/Editor/ModelingToolsEditorMode/Source/ModelingToolsEditorMode/Public/ModelingToolsEditorMode.h
+++ b/Engine/Plugins/Editor/ModelingToolsEditorMode/Source/ModelingToolsEditorMode/Public/ModelingToolsEditorMode.h
@@ -22,7 +22,7 @@ class UGeometrySelectionManager;
 class UInteractiveCommand;
 
 UCLASS(Transient)
-class UModelingToolsEditorMode : public UBaseLegacyWidgetEdMode
+class UModelingToolsEditorMode : public UBaseLegacyWidgetEdMode, public ILegacyEdModeSelectInterface
 {
 	GENERATED_BODY()
 public:
@@ -66,6 +66,12 @@ public:
 	//////////////////
 
 
+	// ILegacyEdModeSelectInterface
+	virtual bool BoxSelect(FBox& InBox, bool InSelect = true) override;
+	virtual bool FrustumSelect(const FConvexVolume& InFrustum, FEditorViewportClient* InViewportClient, bool InSelect) override;
+
+
+
 	//
 	// Selection System configuration, this will likely move elsewhere
 	//
diff --git a/Engine/Plugins/Editor/ModelingToolsEditorMode/Source/ModelingToolsEditorMode/Public/ModelingToolsEditorModeSettings.h b/Engine/Plugins/Editor/ModelingToolsEditorMode/Source/ModelingToolsEditorMode/Public/ModelingToolsEditorModeSettings.h
--- a/Engine/Plugins/Editor/ModelingToolsEditorMode/Source/ModelingToolsEditorMode/Public/ModelingToolsEditorModeSettings.h
+++ b/Engine/Plugins/Editor/ModelingToolsEditorMode/Source/ModelingToolsEditorMode/Public/ModelingToolsEditorModeSettings.h
@@ -325,4 +325,31 @@ public:
 	 */
 	UPROPERTY(config, EditAnywhere, Category = "Modeling Mode|Tool Assets")
 	TArray<FModelingModeAssetCollectionSet> BrushAlphaSets;
+
+
+
+public:
+
+	// saved-state for various mode settings that are configured via UI toggles/etc, and not exposed in settings dialog
+
+	UPROPERTY(config)
+	int32 LastMeshSelectionDragMode = 0;
+
+public:
+
+	// saved-state for various mode settings that does not persist between editor runs
+
+	UPROPERTY()
+	int32 LastMeshSelectionElementType = 0;
+
+	UPROPERTY()
+	int32 LastMeshSelectionTopologyMode = 0;
+
+	UPROPERTY()
+	bool bLastMeshSelectionVolumeToggle = false;
+
+	UPROPERTY()
+	bool bLastMeshSelectionStaticMeshToggle = false;
+
+
 };
diff --git a/Engine/Plugins/Editor/ModelingToolsEditorMode/Source/ModelingToolsEditorMode/Public/ModelingToolsManagerActions.h b/Engine/Plugins/Editor/ModelingToolsEditorMode/Source/ModelingToolsEditorMode/Public/ModelingToolsManagerActions.h
--- a/Engine/Plugins/Editor/ModelingToolsEditorMode/Source/ModelingToolsEditorMode/Public/ModelingToolsManagerActions.h
+++ b/Engine/Plugins/Editor/ModelingToolsEditorMode/Source/ModelingToolsEditorMode/Public/ModelingToolsManagerActions.h
@@ -141,14 +141,7 @@ public:
 	
 
 
-	// old selection state actions
-	TSharedPtr<FUICommandInfo> BeginSelectionAction_ToVertexType;
-	TSharedPtr<FUICommandInfo> BeginSelectionAction_ToEdgeType;
-	TSharedPtr<FUICommandInfo> BeginSelectionAction_ToFaceType;
-	TSharedPtr<FUICommandInfo> BeginSelectionAction_ToObjectType;
-	TSharedPtr<FUICommandInfo> BeginSelectionAction_ToTriangleType;
-	TSharedPtr<FUICommandInfo> BeginSelectionAction_ToPolygroupType;
-
+	// selection state actions
 	TSharedPtr<FUICommandInfo> MeshSelectionModeAction_NoSelection;
 	TSharedPtr<FUICommandInfo> MeshSelectionModeAction_MeshTriangles;
 	TSharedPtr<FUICommandInfo> MeshSelectionModeAction_MeshVertices;
diff --git a/Engine/Plugins/Runtime/MeshModelingToolset/Source/ModelingComponents/Private/Selection/DynamicMeshSelector.cpp b/Engine/Plugins/Runtime/MeshModelingToolset/Source/ModelingComponents/Private/Selection/DynamicMeshSelector.cpp
--- a/Engine/Plugins/Runtime/MeshModelingToolset/Source/ModelingComponents/Private/Selection/DynamicMeshSelector.cpp
+++ b/Engine/Plugins/Runtime/MeshModelingToolset/Source/ModelingComponents/Private/Selection/DynamicMeshSelector.cpp
@@ -349,6 +349,175 @@ void FBaseDynamicMeshSelector::GetSelectionPreviewForRaycast(
 
 
 
+void FBaseDynamicMeshSelector::UpdateSelectionViaShape(
+	const FWorldShapeQueryInfo& ShapeInfo,
+	FGeometrySelectionEditor& SelectionEditor,
+	const FGeometrySelectionUpdateConfig& UpdateConfig,
+	FGeometrySelectionUpdateResult& ResultOut )
+{
+	ResultOut.bSelectionModified = false;
+	ResultOut.bSelectionMissed = true;
+
+	FTransformSRT3d WorldTransform = GetWorldTransformFunc();
+	TArray<uint64> InsideElements;
+
+	// todo: this should be made more efficient using BVHs...
+
+	if ( SelectionEditor.GetTopologyType() == EGeometryTopologyType::Triangle )
+	{
+		if ( SelectionEditor.GetElementType() == EGeometryElementType::Vertex )
+		{
+			GetDynamicMesh()->ProcessMesh([&](const FDynamicMesh3& Mesh) 
+			{
+				TArray<uint64> InsidePoints;
+				for ( int32 vid : Mesh.VertexIndicesItr() )
+				{
+					FVector3d WorldPos = WorldTransform.TransformPosition(Mesh.GetVertex(vid));
+					if ( ShapeInfo.Convex.IntersectPoint(WorldPos) )
+					{
+						InsideElements.Add( FGeoSelectionID::MeshVertex(vid).Encoded() );
+					}
+				}
+			});
+		}
+		else if ( SelectionEditor.GetElementType() == EGeometryElementType::Edge )
+		{
+			GetDynamicMesh()->ProcessMesh([&](const FDynamicMesh3& Mesh) 
+			{
+				TArray<uint64> InsideEdges;
+				for ( int32 eid : Mesh.EdgeIndicesItr() )
+				{
+					FVector3d A,B;
+					Mesh.GetEdgeV(eid, A,B);
+					A = WorldTransform.TransformPosition(A);
+					B = WorldTransform.TransformPosition(B);
+					bool bFullyContained = false;
+					if ( ShapeInfo.Convex.IntersectLineSegment(A,B) )
+					{
+						bFullyContained = ShapeInfo.Convex.IntersectPoint(A) && ShapeInfo.Convex.IntersectPoint(B);
+						if (bFullyContained)
+						{
+							InsideElements.Add( FGeoSelectionID::MeshEdge( Mesh.GetTriEdgeIDFromEdgeID(eid) ).Encoded() );
+						}
+					}
+				}
+			});
+		}
+		else if (SelectionEditor.GetElementType() == EGeometryElementType::Face)
+		{
+			GetDynamicMesh()->ProcessMesh([&](const FDynamicMesh3& Mesh) 
+			{
+				TArray<uint64> InsideFaces;
+				for ( int32 tid : Mesh.TriangleIndicesItr() )
+				{
+					FVector3d A,B,C;
+					Mesh.GetTriVertices(tid, A,B,C);
+					A = WorldTransform.TransformPosition(A);
+					B = WorldTransform.TransformPosition(B);
+					C = WorldTransform.TransformPosition(C);
+					bool bFullyContained = false;
+					if ( ShapeInfo.Convex.IntersectTriangle(A,B,C, bFullyContained) )
+					{
+						if (bFullyContained)
+						{
+							InsideElements.Add( FGeoSelectionID::MeshTriangle(tid).Encoded() );
+						}
+					}
+				}
+			});
+		}
+	}
+	else if ( SelectionEditor.GetTopologyType() == EGeometryTopologyType::Polygroup )
+	{
+		const FGroupTopology* UseGroupTopology = GetGroupTopology();
+
+		if ( SelectionEditor.GetElementType() == EGeometryElementType::Vertex )
+		{
+			GetDynamicMesh()->ProcessMesh([&](const FDynamicMesh3& Mesh) 
+			{
+				for (const FGroupTopology::FCorner& Corner : UseGroupTopology->Corners)
+				{
+					FVector3d WorldPos = WorldTransform.TransformPosition(Mesh.GetVertex(Corner.VertexID));
+					if ( ShapeInfo.Convex.IntersectPoint(WorldPos) )
+					{
+						InsideElements.Add( FGeoSelectionID(Corner.VertexID, UseGroupTopology->GetCornerIDFromVertexID(Corner.VertexID)).Encoded() );
+					}				
+				}
+			});
+
+		}
+		else if (SelectionEditor.GetElementType() == EGeometryElementType::Edge)
+		{
+			GetDynamicMesh()->ProcessMesh([&](const FDynamicMesh3& Mesh) 
+			{
+				int32 NumEdges = UseGroupTopology->Edges.Num();
+				for ( int32 GroupEdgeID = 0; GroupEdgeID < NumEdges; ++GroupEdgeID)
+				{
+					const FGroupTopology::FGroupEdge& Edge = UseGroupTopology->Edges[GroupEdgeID];
+					bool bAllInside = true;
+					for (int32 eid : Edge.Span.Edges)
+					{
+						FVector3d A,B;
+						Mesh.GetEdgeV(eid, A,B);
+						A = WorldTransform.TransformPosition(A);
+						B = WorldTransform.TransformPosition(B);
+						bool bFullyContained = false;
+						bool bIntersects = ShapeInfo.Convex.IntersectLineSegment(A,B);
+						bFullyContained = bIntersects && ShapeInfo.Convex.IntersectPoint(A) && ShapeInfo.Convex.IntersectPoint(B);
+						if ( bFullyContained == false || bFullyContained == false)
+						{
+							bAllInside = false;
+							break;
+						}
+					}
+					if (bAllInside)
+					{
+						FMeshTriEdgeID EdgeTriID = Mesh.GetTriEdgeIDFromEdgeID(Edge.Span.Edges[0]);
+						InsideElements.Add( FGeoSelectionID( EdgeTriID.Encoded(), GroupEdgeID).Encoded());
+					}
+				}
+			});
+		}
+		else if (SelectionEditor.GetElementType() == EGeometryElementType::Face)
+		{
+			GetDynamicMesh()->ProcessMesh([&](const FDynamicMesh3& Mesh) 
+			{
+				for (const FGroupTopology::FGroup& Group : UseGroupTopology->Groups)
+				{
+					bool bAllInside = true;
+					for (int32 tid : Group.Triangles)
+					{
+						FVector3d A,B,C;
+						Mesh.GetTriVertices(tid, A,B,C);
+						A = WorldTransform.TransformPosition(A);
+						B = WorldTransform.TransformPosition(B);
+						C = WorldTransform.TransformPosition(C);
+						bool bFullyContained = false;
+						bool bIntersects = ShapeInfo.Convex.IntersectTriangle(A,B,C, bFullyContained);
+						if ( bIntersects == false || bFullyContained == false)
+						{
+							bAllInside = false;
+							break;
+						}
+					}
+					if (bAllInside)
+					{
+						InsideElements.Add( FGeoSelectionID(Group.Triangles[0], Group.GroupID).Encoded());
+					}
+				}
+			});
+		}
+	}
+
+	if (InsideElements.Num() > 0 )
+	{
+		ResultOut.bSelectionModified = UpdateSelectionWithNewElements(&SelectionEditor, UpdateConfig.ChangeType, InsideElements, &ResultOut.SelectionDelta);
+		ResultOut.bSelectionMissed = false;
+	}
+}
+
+
+
 void FBaseDynamicMeshSelector::GetSelectionFrame(const FGeometrySelection& Selection, FFrame3d& SelectionFrame, bool bTransformToWorld)
 {
 	if (Selection.TopologyType == EGeometryTopologyType::Polygroup)
diff --git a/Engine/Plugins/Runtime/MeshModelingToolset/Source/ModelingComponents/Private/Selection/GeometrySelectionManager.cpp b/Engine/Plugins/Runtime/MeshModelingToolset/Source/ModelingComponents/Private/Selection/GeometrySelectionManager.cpp
--- a/Engine/Plugins/Runtime/MeshModelingToolset/Source/ModelingComponents/Private/Selection/GeometrySelectionManager.cpp
+++ b/Engine/Plugins/Runtime/MeshModelingToolset/Source/ModelingComponents/Private/Selection/GeometrySelectionManager.cpp
@@ -567,6 +567,50 @@ void UGeometrySelectionManager::UpdateSelectionViaRaycast(
 }
 
 
+void UGeometrySelectionManager::UpdateSelectionViaConvex(
+	const FConvexVolume& ConvexVolume,
+	const FGeometrySelectionUpdateConfig& UpdateConfig,
+	FGeometrySelectionUpdateResult& ResultOut )
+{
+	ResultOut.bSelectionModified = false;
+
+	if (ActiveTargetReferences.Num() == 0)
+	{
+		return;
+	}
+
+	// currently only going to support one object, not sure how to support more yet...
+	FGeometrySelectionTarget* Target = ActiveTargetReferences[0].Get();
+
+	IGeometrySelector::FWorldShapeQueryInfo ShapeQueryInfo;
+	ShapeQueryInfo.Convex = ConvexVolume;
+	IToolsContextQueriesAPI* QueryAPI = this->ToolsContext->ToolManager->GetContextQueriesAPI();
+	QueryAPI->GetCurrentViewState(ShapeQueryInfo.CameraState);
+
+	Target->Selector->UpdateSelectionViaShape(
+		ShapeQueryInfo, *Target->SelectionEditor, UpdateConfig, ResultOut );
+
+	if (ResultOut.bSelectionModified)
+	{
+		TUniquePtr<FGeometrySelectionDeltaChange> DeltaChange = MakeUnique<FGeometrySelectionDeltaChange>();
+		DeltaChange->Identifier = Target->SelectionIdentifer;
+		DeltaChange->Delta = ResultOut.SelectionDelta;
+
+		GetTransactionsAPI()->BeginUndoTransaction(LOCTEXT("UpdateSelectionViaConvex", "Change Selection"));
+		GetTransactionsAPI()->AppendChange(this, MoveTemp(DeltaChange), LOCTEXT("UpdateSelectionViaConvex", "Change Selection"));
+		GetTransactionsAPI()->EndUndoTransaction();
+
+		bSelectionRenderCachesDirty = true;
+		OnSelectionModified.Broadcast();
+	}
+	else if (ResultOut.bSelectionMissed && UpdateConfig.ChangeType == EGeometrySelectionChangeType::Replace)
+	{
+		ClearSelection();
+	}
+}
+
+
+
 bool UGeometrySelectionManager::CanBeginTrackedSelectionChange() const
 {
 	return ActiveTargetReferences.Num() > 0 && bInTrackedSelectionChange == false;
diff --git a/Engine/Plugins/Runtime/MeshModelingToolset/Source/ModelingComponents/Public/Selection/DynamicMeshSelector.h b/Engine/Plugins/Runtime/MeshModelingToolset/Source/ModelingComponents/Public/Selection/DynamicMeshSelector.h
--- a/Engine/Plugins/Runtime/MeshModelingToolset/Source/ModelingComponents/Public/Selection/DynamicMeshSelector.h
+++ b/Engine/Plugins/Runtime/MeshModelingToolset/Source/ModelingComponents/Public/Selection/DynamicMeshSelector.h
@@ -99,6 +99,16 @@ public:
 		const FWorldRayQueryInfo& RayInfo,
 		FGeometrySelectionEditor& PreviewEditor) override;
 
+
+
+	virtual void UpdateSelectionViaShape(	
+		const FWorldShapeQueryInfo& ShapeInfo,
+		FGeometrySelectionEditor& SelectionEditor,
+		const FGeometrySelectionUpdateConfig& UpdateConfig,
+		FGeometrySelectionUpdateResult& ResultOut 
+	) override;
+
+
 	virtual FTransform GetLocalToWorldTransform() const
 	{
 		return (FTransform)GetWorldTransformFunc();
diff --git a/Engine/Plugins/Runtime/MeshModelingToolset/Source/ModelingComponents/Public/Selection/GeometrySelectionManager.h b/Engine/Plugins/Runtime/MeshModelingToolset/Source/ModelingComponents/Public/Selection/GeometrySelectionManager.h
--- a/Engine/Plugins/Runtime/MeshModelingToolset/Source/ModelingComponents/Public/Selection/GeometrySelectionManager.h
+++ b/Engine/Plugins/Runtime/MeshModelingToolset/Source/ModelingComponents/Public/Selection/GeometrySelectionManager.h
@@ -3,6 +3,7 @@
 #pragma once
 
 #include "CoreMinimal.h"
+#include "ConvexVolume.h"
 #include "Selections/GeometrySelection.h"
 #include "Selection/GeometrySelector.h"
 #include "Selection/GeometrySelectionChanges.h"
@@ -169,6 +170,18 @@ public:
 	);
 
 
+	/**
+	 * Use the given ConvexVolume to update the active element selection based on UpdateConfig.
+	 * @param ResultOut information on any element selection modifications is returned here
+	 */
+	virtual void UpdateSelectionViaConvex(
+		const FConvexVolume& ConvexVolume,
+		const FGeometrySelectionUpdateConfig& UpdateConfig,
+		FGeometrySelectionUpdateResult& ResultOut
+	);
+
+
+
 	//
 	// Support for more complex selection changes that might (eg) occur over multiple frames, 
 	// or be computed externally. The usage pattern is:
diff --git a/Engine/Plugins/Runtime/MeshModelingToolset/Source/ModelingComponents/Public/Selection/GeometrySelector.h b/Engine/Plugins/Runtime/MeshModelingToolset/Source/ModelingComponents/Public/Selection/GeometrySelector.h
--- a/Engine/Plugins/Runtime/MeshModelingToolset/Source/ModelingComponents/Public/Selection/GeometrySelector.h
+++ b/Engine/Plugins/Runtime/MeshModelingToolset/Source/ModelingComponents/Public/Selection/GeometrySelector.h
@@ -3,6 +3,7 @@
 #pragma once
 
 #include "Components/PrimitiveComponent.h"
+#include "ConvexVolume.h"
 #include "Selections/GeometrySelection.h"
 #include "FrameTypes.h"
 #include "InputState.h"
@@ -345,6 +346,27 @@ public:
 		FGeometrySelectionEditor& PreviewEditor
 	) = 0;
 
+
+
+	struct FWorldShapeQueryInfo
+	{
+		FConvexVolume Convex;
+		FViewCameraState CameraState;
+	};
+
+	/**
+	 * Update the active selection based on a world-space Shape.
+	 * Uses the provided SelectionEditor and UpdateConfig to do the update
+	 * Information about any changes actually made to the selection via the SelectionEditor are returned in ResultOut
+	 */
+	virtual void UpdateSelectionViaShape(	
+		const FWorldShapeQueryInfo& ShapeInfo,
+		FGeometrySelectionEditor& SelectionEditor,
+		const FGeometrySelectionUpdateConfig& UpdateConfig,
+		FGeometrySelectionUpdateResult& ResultOut 
+	) = 0;
+
+
 	/**
 	 * @return the World transform for the Selector's target object
 	 */
