ModelingMode: Add experimental support for tool presets.

* Presets exist as a new asset type, loaded under the new ToolPresets plugin, a dependency of ModelingToolsEditorMode.
* Presets consist of a nested set of maps, containing tool specific, named presets. These objects then contain copies of the tool in question properties at the time of creation, allowing the preset to later replay the settings back onto the active tool.
* Additionally, behind an experimental CVAR "modeling.EnablePresets", UX for interacting with presets within the Modeling Mode is available, allowing for selecting, saving, applying, and renaming new presets.

diff --git a/Engine/Plugins/Editor/ModelingToolsEditorMode/ModelingToolsEditorMode.uplugin b/Engine/Plugins/Editor/ModelingToolsEditorMode/ModelingToolsEditorMode.uplugin
--- a/Engine/Plugins/Editor/ModelingToolsEditorMode/ModelingToolsEditorMode.uplugin
+++ b/Engine/Plugins/Editor/ModelingToolsEditorMode/ModelingToolsEditorMode.uplugin
@@ -39,6 +39,11 @@
 		{
 			"Name": "StylusInput",
 			"Enabled": true
+		},
+		{
+			"Name": "ToolPresets",
+			"Enabled": true,
+			"TargetAllowList": [ "Editor" ]
 		}
 	]
 }
\ No newline at end of file
diff --git a/Engine/Plugins/Editor/ModelingToolsEditorMode/Source/ModelingToolsEditorMode/ModelingToolsEditorMode.Build.cs b/Engine/Plugins/Editor/ModelingToolsEditorMode/Source/ModelingToolsEditorMode/ModelingToolsEditorMode.Build.cs
--- a/Engine/Plugins/Editor/ModelingToolsEditorMode/Source/ModelingToolsEditorMode/ModelingToolsEditorMode.Build.cs
+++ b/Engine/Plugins/Editor/ModelingToolsEditorMode/Source/ModelingToolsEditorMode/ModelingToolsEditorMode.Build.cs
@@ -49,6 +49,7 @@ public class ModelingToolsEditorMode : ModuleRules
 				"StylusInput",
 				"DeveloperSettings",
 				"PropertyEditor",
+				"PresetAsset"
 				// ... add private dependencies that you statically link with here ...	
 			}
 		);
diff --git a/Engine/Plugins/Editor/ModelingToolsEditorMode/Source/ModelingToolsEditorMode/Private/ModelingToolsEditorModeToolkit.cpp b/Engine/Plugins/Editor/ModelingToolsEditorMode/Source/ModelingToolsEditorMode/Private/ModelingToolsEditorModeToolkit.cpp
--- a/Engine/Plugins/Editor/ModelingToolsEditorMode/Source/ModelingToolsEditorMode/Private/ModelingToolsEditorModeToolkit.cpp
+++ b/Engine/Plugins/Editor/ModelingToolsEditorMode/Source/ModelingToolsEditorMode/Private/ModelingToolsEditorModeToolkit.cpp
@@ -18,6 +18,8 @@
 #include "SSimpleButton.h"
 #include "Tools/UEdMode.h"
 #include "Widgets/Layout/SExpandableArea.h"
+#include "Widgets/Input/SEditableComboBox.h"
+#include "Widgets/Input/STextEntryPopup.h"
 #include "Framework/MultiBox/SToolBarButtonBlock.h"
 #include "ModelingWidgets/ModelingCustomizationUtil.h"
 
@@ -37,6 +39,10 @@
 #include "Dialogs/DlgPickPath.h"
 #include "ModelingModeAssetUtils.h"
 
+// for presets
+#include "PresetAsset.h"
+#include "PropertyCustomizationHelpers.h"
+
 #include "LevelEditorViewport.h"
 #include "Widgets/Input/SEditableTextBox.h"
 #include "Widgets/Input/STextComboBox.h"
@@ -54,7 +60,31 @@ static TAutoConsoleVariable<int32> CVarEnablePolyModeling(
 	TEXT("modeling.EnablePolyModel"),
 	0,
 	TEXT("Enable prototype PolyEdit tab"));
+static TAutoConsoleVariable<int32> CVarEnableToolPresets(
+	TEXT("modeling.EnablePresets"),
+	0,
+	TEXT("Enable tool preset features and UX"));
 
+namespace FModelingToolsEditorModeToolkitLocals
+{
+	typedef TFunction<void(UPresetAsset& Preset, UInteractiveTool& Tool)> PresetAndToolFunc;
+
+	void ExecuteWithPresetAndTool(UEdMode& EdMode, EToolSide ToolSide, FSoftObjectPath& PresetPath, PresetAndToolFunc Function)
+	{
+		UPresetAsset* Preset = nullptr;		
+		UInteractiveTool* Tool = EdMode.GetToolManager()->GetActiveTool(EToolSide::Left);
+
+		if (PresetPath.IsAsset())
+		{
+			Preset = Cast<UPresetAsset>(PresetPath.TryLoad());
+		}
+		if (!Preset || !Tool)
+		{
+			return;
+		}
+		Function(*Preset, *Tool);
+	}
+}
 
 
 FModelingToolsEditorModeToolkit::FModelingToolsEditorModeToolkit()
@@ -405,6 +435,73 @@ TSharedPtr<SWidget> FModelingToolsEditorModeToolkit::MakeAssetConfigPanel()
 		];
 	}
 
+	bool bEnableToolPresets = (CVarEnableToolPresets.GetValueOnGameThread() > 0);
+
+	if(bEnableToolPresets)
+	{
+		BuildPresetComboList();
+		Content->AddSlot().HAlign(HAlign_Fill).VAlign(VAlign_Fill).AutoHeight()[
+			SNew(SHorizontalBox)
+			+ SHorizontalBox::Slot().HAlign(HAlign_Right).Padding(0, 2, 2, 2).AutoWidth()
+			[
+				SNew(STextBlock)
+				.Text(LOCTEXT("PresetLabel", "Preset"))
+			]
+			+ SHorizontalBox::Slot().HAlign(HAlign_Fill).Padding(0, 0, 2, 0).FillWidth(1.0f)
+			[
+				SNew(SVerticalBox)			
+				+ SVerticalBox::Slot().AutoHeight()
+				[
+					SNew(SHorizontalBox)
+					+ SHorizontalBox::Slot().HAlign(HAlign_Fill).Padding(0, 2, 2, 2).FillWidth(1.0f)
+					[
+						SNew(SObjectPropertyEntryBox)
+						.AllowedClass(UPresetAsset::StaticClass())
+						.ObjectPath_Lambda([this]() {return FModelingToolsEditorModeToolkit::GetCurrentPresetPath(); })
+						.OnShouldFilterAsset(this, &FModelingToolsEditorModeToolkit::HandleFilterPresetAsset)
+						.OnObjectChanged(this, &FModelingToolsEditorModeToolkit::HandlePresetAssetChanged)
+					]
+				]
+				+ SVerticalBox::Slot().Padding(0, 2, 2, 2).AutoHeight()
+				[
+					SNew(SHorizontalBox)
+					+ SHorizontalBox::Slot().HAlign(HAlign_Fill).Padding(0, 2, 2, 2).FillWidth(1.0f)
+					[
+						SAssignNew(PresetComboBox, SEditableComboBox<TSharedPtr<FString>>)
+						.OptionsSource(&AvailablePresetsForTool)
+						.OnGenerateWidget(this, &FModelingToolsEditorModeToolkit::MakePresetComboWidget)
+						.OnSelectionChanged(this, &FModelingToolsEditorModeToolkit::OnPresetChanged)
+						//.OnComboBoxOpening(this, &FModelingToolsEditorModeToolkit::BuildPresetComboList)
+						.OnGetEditableText(this, &FModelingToolsEditorModeToolkit::GetRawPresetString)
+						.OnSelectionRenamed(this, &FModelingToolsEditorModeToolkit::OnPresetRenamed)
+						.OnAddClicked(this, &FModelingToolsEditorModeToolkit::OnAddPreset)
+						.OnRemoveClicked(this, &FModelingToolsEditorModeToolkit::OnRemovePreset)
+						.IsEnabled(this, &FModelingToolsEditorModeToolkit::IsPresetEnabled)
+						.ContentPadding(2)
+						.Content()
+						[
+							SNew(STextBlock)
+							.Text(this, &FModelingToolsEditorModeToolkit::GetPresetComboBoxContent)
+							.ToolTipText(this, &FModelingToolsEditorModeToolkit::GetPresetComboBoxContent)
+						]
+					]
+					+ SHorizontalBox::Slot().HAlign(HAlign_Right).Padding(0, 2, 2, 2).AutoWidth()
+					[
+						SNew(SSimpleButton)
+						.OnClicked_Lambda([this]() { LoadActivePreset(); return FReply::Handled(); })
+					.Icon(FAppStyle::Get().GetBrush("Icons.Load"))
+					]
+					+ SHorizontalBox::Slot().HAlign(HAlign_Right).Padding(0, 2, 2, 2).AutoWidth()
+					[
+						SNew(SSimpleButton)
+						.OnClicked_Lambda([this]() { SaveActivePreset(); return FReply::Handled(); })
+					.Icon(FAppStyle::Get().GetBrush("Icons.Save"))
+					]
+				]
+			]
+		];
+	}
+
 	TSharedPtr<SExpandableArea> AssetConfigPanel = SNew(SExpandableArea)
 		.HeaderPadding(FMargin(2.0f))
 		.Padding(FMargin(2.f))
@@ -428,6 +525,246 @@ TSharedPtr<SWidget> FModelingToolsEditorModeToolkit::MakeAssetConfigPanel()
 
 }
 
+TSharedRef<SWidget> FModelingToolsEditorModeToolkit::MakePresetComboWidget(TSharedPtr<FString> InItem)
+{
+	return
+		SNew(STextBlock)
+		.Text(FText::FromString(*InItem));
+}
+
+bool FModelingToolsEditorModeToolkit::IsPresetEnabled() const
+{
+	return CurrentPreset.IsAsset();
+}
+
+void FModelingToolsEditorModeToolkit::OnPresetRenamed(const FText& NewPresetName, ETextCommit::Type CommitType)
+{
+	FModelingToolsEditorModeToolkitLocals::ExecuteWithPresetAndTool(*OwningEditorMode, EToolSide::Left, CurrentPreset,
+		[this, NewPresetName] (UPresetAsset& Preset, UInteractiveTool& Tool) {
+		FPropertyStore PresetValuesToRename;
+		if (!Preset.StoredProperties.Contains(Tool.GetClass()->GetName()))
+		{
+			return;
+		}
+		if (Preset.StoredProperties[Tool.GetClass()->GetName()].Store.Contains(*GetPresetString()))
+		{
+			Preset.StoredProperties[Tool.GetClass()->GetName()].Store.RemoveAndCopyValue(*GetPresetString(), PresetValuesToRename);
+			Preset.StoredProperties[Tool.GetClass()->GetName()].Store.Add(NewPresetName.ToString(), PresetValuesToRename);
+			BuildPresetComboList();
+			ActiveNamedPreset = *AvailablePresetsForTool.FindByPredicate([this, NewPresetName](TSharedPtr<FString>& Value) { return Value->Equals(NewPresetName.ToString()); });
+
+			Preset.MarkPackageDirty();
+		}
+	});
+}
+
+FReply FModelingToolsEditorModeToolkit::OnAddPreset()
+{
+	FModelingToolsEditorModeToolkitLocals::ExecuteWithPresetAndTool(*OwningEditorMode, EToolSide::Left, CurrentPreset,
+		[this] (UPresetAsset& Preset, UInteractiveTool& Tool) {
+
+		FString NewPresetName = FString::Printf(TEXT("NewPreset_%d"), Preset.StoredProperties.FindOrAdd(Tool.GetClass()->GetName()).Store.Num());
+		FPropertyStore PresetValuesToCreate;
+		Preset.StoredProperties[Tool.GetClass()->GetName()].Store.Add(NewPresetName, PresetValuesToCreate);
+		BuildPresetComboList();
+		ActiveNamedPreset = *AvailablePresetsForTool.FindByPredicate([this, NewPresetName](TSharedPtr<FString>& Value) { return Value->Equals(NewPresetName); });
+
+		Preset.MarkPackageDirty();
+	});
+	return FReply::Handled();
+}
+
+FReply FModelingToolsEditorModeToolkit::OnRemovePreset()
+{
+	FModelingToolsEditorModeToolkitLocals::ExecuteWithPresetAndTool(*OwningEditorMode, EToolSide::Left, CurrentPreset,
+		[this] (UPresetAsset& Preset, UInteractiveTool& Tool) {
+		if (!Preset.StoredProperties.Contains(Tool.GetClass()->GetName()))
+		{
+			return;
+		}
+		FPropertyStore PresetValuesToRename;
+		if (Preset.StoredProperties[Tool.GetClass()->GetName()].Store.Contains(*GetPresetString()))
+		{
+			Preset.StoredProperties[Tool.GetClass()->GetName()].Store.RemoveAndCopyValue(*GetPresetString(), PresetValuesToRename);
+			ActiveNamedPreset.Reset();
+			BuildPresetComboList();
+
+			Preset.MarkPackageDirty();
+		}
+	});
+	return FReply::Handled();
+}
+
+void FModelingToolsEditorModeToolkit::OnPresetChanged(TSharedPtr<FString> NewSelection, ESelectInfo::Type SelectInfo)
+{
+	ActiveNamedPreset = NewSelection;
+}
+
+
+FText FModelingToolsEditorModeToolkit::GetPresetComboBoxContent() const
+{
+	TSharedPtr<FString> CurrentPresetValue = GetPresetString();
+	if (CurrentPresetValue)
+	{
+		return FText::FromString(*CurrentPresetValue.Get());
+	}
+	else
+	{
+		return FText::GetEmpty();
+	}
+}
+
+
+TSharedPtr<FString> FModelingToolsEditorModeToolkit::GetPresetString() const
+{
+	return ActiveNamedPreset;
+}
+
+FString FModelingToolsEditorModeToolkit::GetRawPresetString() const
+{
+	return GetPresetString() ? *GetPresetString() : TEXT("");
+}
+
+void FModelingToolsEditorModeToolkit::BuildPresetComboList()
+{
+	FModelingToolsEditorModeToolkitLocals::ExecuteWithPresetAndTool(*OwningEditorMode, EToolSide::Left, CurrentPreset,
+		[this] (UPresetAsset& Preset, UInteractiveTool& Tool) {
+
+		AvailablePresetsForTool.Empty();
+		if (!Preset.StoredProperties.Contains(Tool.GetClass()->GetName()))
+		{
+			PresetComboBox->RefreshOptions();
+			PresetComboBox->ClearSelection();
+			return;
+		}
+		TArray<FString> PresetsForTool;
+		Preset.StoredProperties[Tool.GetClass()->GetName()].Store.GenerateKeyArray(PresetsForTool);
+		AvailablePresetsForTool.Reserve(PresetsForTool.Num());
+		for (FString& PresetEntry : PresetsForTool)
+		{
+			AvailablePresetsForTool.Add(MakeShared<FString>(PresetEntry));
+		}
+		PresetComboBox->RefreshOptions();
+		PresetComboBox->ClearSelection();
+	});
+}
+
+void FModelingToolsEditorModeToolkit::ClearPresetComboList()
+{
+	ActiveNamedPreset.Reset();
+	AvailablePresetsForTool.Empty();
+	PresetComboBox->RefreshOptions();
+	PresetComboBox->ClearSelection();
+}
+
+void FModelingToolsEditorModeToolkit::HandlePresetAssetChanged(const FAssetData& InAssetData)
+{
+	CurrentPreset.SetPath(InAssetData.GetObjectPathString());
+	if (CurrentPreset.IsAsset())
+	{
+		BuildPresetComboList();
+	}
+	else
+	{
+		CurrentPreset.Reset();
+	}
+}
+
+bool FModelingToolsEditorModeToolkit::HandleFilterPresetAsset(const FAssetData& InAssetData)
+{
+#if 1
+	return false;
+#else 
+    // This was the "old" approach, where only presets valid for the current tool were shown. But this precluded selecting
+    // other assets and *adding* the current tool to them, so the filtering was removed. However, the logic may prove useful,
+    // so this is just keeping it around for a little while longer just in case.
+
+	UPresetAsset* PresetAsset = Cast<UPresetAsset>(InAssetData.GetAsset());
+	if (PresetAsset) 
+	{
+		bool bMatchingToolName = false;
+		if (OwningEditorMode->GetToolManager()->HasActiveTool(EToolSide::Left))
+		{
+			bMatchingToolName |= PresetAsset->StoredProperties.Contains(OwningEditorMode->GetToolManager()->GetActiveTool(EToolSide::Left)->GetClass()->GetName());
+		}
+		if (OwningEditorMode->GetToolManager()->HasActiveTool(EToolSide::Right))
+		{
+			bMatchingToolName |= PresetAsset->StoredProperties.Contains(OwningEditorMode->GetToolManager()->GetActiveTool(EToolSide::Right)->GetClass()->GetName());
+		}
+		return !bMatchingToolName;
+	}
+
+	return true;
+#endif
+}
+
+void FModelingToolsEditorModeToolkit::LoadActivePreset()
+{
+	FModelingToolsEditorModeToolkitLocals::ExecuteWithPresetAndTool(*OwningEditorMode, EToolSide::Left, CurrentPreset,
+	[this](UPresetAsset& Preset, UInteractiveTool& Tool) {
+		if (!Preset.StoredProperties.Contains(Tool.GetClass()->GetName()))
+		{
+			return;
+		}
+
+		TArray<UObject*> PropertySets = Tool.GetToolProperties();
+		TArray<TObjectPtr<UObject>>& PropertyStore = Preset.StoredProperties[Tool.GetClass()->GetName()].Store.FindOrAdd(*ActiveNamedPreset).Properties;
+
+		for (UObject* PropertySet : PropertySets)
+		{
+			UObject* StoredPropertySet = nullptr;
+			for (UObject* PotentialStoredPropertySet : PropertyStore)
+			{
+				if (PotentialStoredPropertySet && PotentialStoredPropertySet->GetClass() == PropertySet->GetClass())
+				{
+					StoredPropertySet = PotentialStoredPropertySet;
+				}
+			}
+			if(StoredPropertySet != nullptr)
+			{
+				for (FProperty* Prop : TFieldRange<FProperty>(PropertySet->GetClass()))
+				{
+					void* DestValue = Prop->ContainerPtrToValuePtr<void>(PropertySet);
+					void* SrcValue = Prop->ContainerPtrToValuePtr<void>(StoredPropertySet);
+					if (!Prop->Identical(DestValue, SrcValue))
+					{
+						Prop->CopySingleValue(DestValue, SrcValue);
+						Tool.OnPropertyModified(PropertySet, Prop);
+					}
+				}
+			}
+		}
+	});
+}
+
+void FModelingToolsEditorModeToolkit::SaveActivePreset()
+{
+
+	FModelingToolsEditorModeToolkitLocals::ExecuteWithPresetAndTool(*OwningEditorMode, EToolSide::Left, CurrentPreset,
+	[this](UPresetAsset& Preset, UInteractiveTool& Tool) {
+		TArray<UObject*> PropertySets = Tool.GetToolProperties();
+		TArray<TObjectPtr<UObject>>& PropertyStore = Preset.StoredProperties.FindOrAdd(Tool.GetClass()->GetName()).Store.FindOrAdd(*ActiveNamedPreset).Properties;
+		PropertyStore.Empty();
+
+		for (UObject* PropertySet : PropertySets)
+		{
+			UObject* StoredPropertySet = NewObject<UObject>(&Preset, PropertySet->GetClass());
+			StoredPropertySet->ClearFlags(EObjectFlags::RF_Transient);
+			if (StoredPropertySet != nullptr)
+			{
+				for (FProperty* Prop : TFieldRange<FProperty>(PropertySet->GetClass()))
+				{
+					void* SrcValue = Prop->ContainerPtrToValuePtr<void>(PropertySet);
+					void* DestValue = Prop->ContainerPtrToValuePtr<void>(StoredPropertySet);
+					Prop->CopySingleValue(DestValue, SrcValue);
+				}
+			}
+			PropertyStore.Add(StoredPropertySet);
+		}	
+
+		Preset.MarkPackageDirty();
+	});
+}
 
 void FModelingToolsEditorModeToolkit::InitializeAfterModeSetup()
 {
@@ -1114,6 +1451,8 @@ void FModelingToolsEditorModeToolkit::OnToolStarted(UInteractiveToolManager* Man
 			Viewport->Invalidate();
 		}
 	}
+
+	BuildPresetComboList();
 }
 
 void FModelingToolsEditorModeToolkit::OnToolEnded(UInteractiveToolManager* Manager, UInteractiveTool* Tool)
@@ -1139,6 +1478,8 @@ void FModelingToolsEditorModeToolkit::OnToolEnded(UInteractiveToolManager* Manag
 	// re-enable LOD level picker
 	AssetLODMode->SetEnabled(true);
 	AssetLODModeLabel->SetEnabled(true);
+
+	ClearPresetComboList();
 }
 
 void FModelingToolsEditorModeToolkit::OnActiveViewportChanged(TSharedPtr<IAssetViewport> OldViewport, TSharedPtr<IAssetViewport> NewViewport)
diff --git a/Engine/Plugins/Editor/ModelingToolsEditorMode/Source/ModelingToolsEditorMode/Public/ModelingToolsEditorModeToolkit.h b/Engine/Plugins/Editor/ModelingToolsEditorMode/Source/ModelingToolsEditorMode/Public/ModelingToolsEditorModeToolkit.h
--- a/Engine/Plugins/Editor/ModelingToolsEditorMode/Source/ModelingToolsEditorMode/Public/ModelingToolsEditorModeToolkit.h
+++ b/Engine/Plugins/Editor/ModelingToolsEditorMode/Source/ModelingToolsEditorMode/Public/ModelingToolsEditorModeToolkit.h
@@ -112,6 +112,31 @@ private:
 	void OnShowAssetSettings();
 	void SelectNewAssetPath() const;
 
+	// Presets
+	FSoftObjectPath CurrentPreset;
+	TSharedPtr<FString> ActiveNamedPreset;
+	TArray<TSharedPtr<FString>> AvailablePresetsForTool;
+	TSharedPtr<SEditableComboBox<TSharedPtr<FString>>> PresetComboBox;
+
+	FString GetCurrentPresetPath() { return CurrentPreset.GetAssetPathString(); }
+	void HandlePresetAssetChanged(const FAssetData& InAssetData);
+	bool HandleFilterPresetAsset(const FAssetData& InAssetData);
+	void LoadActivePreset();
+	void SaveActivePreset();
+	
+	TSharedRef<SWidget> MakePresetComboWidget(TSharedPtr<FString> InItem);
+	bool IsPresetEnabled() const;
+	void OnPresetChanged(TSharedPtr<FString> NewSelection, ESelectInfo::Type SelectInfo);
+	void OnPresetRenamed(const FText&, ETextCommit::Type);
+	FReply OnAddPreset();
+	FReply OnRemovePreset();
+	FText GetPresetComboBoxContent() const;
+	TSharedPtr<FString> GetPresetString() const;
+	FString GetRawPresetString() const;
+	void BuildPresetComboList();
+	void ClearPresetComboList();
+
+
 	TArray<TSharedPtr<FString>> AssetLODModes;
 	TSharedPtr<STextBlock> AssetLODModeLabel;
 	TSharedPtr<STextComboBox> AssetLODMode;
diff --git a/Engine/Plugins/Experimental/ToolPresets/Source/PresetAsset/PresetAsset.Build.cs b/Engine/Plugins/Experimental/ToolPresets/Source/PresetAsset/PresetAsset.Build.cs
new file mode 100644
--- /dev/null
+++ b/Engine/Plugins/Experimental/ToolPresets/Source/PresetAsset/PresetAsset.Build.cs
@@ -0,0 +1,53 @@
+// Copyright Epic Games, Inc. All Rights Reserved.
+
+using UnrealBuildTool;
+
+public class PresetAsset: ModuleRules
+{
+	public PresetAsset(ReadOnlyTargetRules Target) : base(Target)
+	{
+		PCHUsage = ModuleRules.PCHUsageMode.UseExplicitOrSharedPCHs;
+		
+		PublicIncludePaths.AddRange(
+			new string[] {
+				// ... add public include paths required here ...
+			}
+			);
+				
+		
+		PrivateIncludePaths.AddRange(
+			new string[] {
+				// ... add other private include paths required here ...
+			}
+			);
+			
+		
+		PublicDependencyModuleNames.AddRange(
+			new string[]
+			{
+				"Core",
+			}
+			);
+			
+		
+		PrivateDependencyModuleNames.AddRange(
+			new string[]
+			{
+				"CoreUObject",
+				"Engine",
+				"Slate",
+				"SlateCore",
+				"UnrealEd",
+				"AssetDefinition",
+			}
+			);
+		
+		
+		DynamicallyLoadedModuleNames.AddRange(
+			new string[]
+			{
+				// ... add any modules that your module loads dynamically here ...
+			}
+			);
+	}
+}
diff --git a/Engine/Plugins/Experimental/ToolPresets/Source/PresetAsset/Private/AssetDefinition_PresetAsset.cpp b/Engine/Plugins/Experimental/ToolPresets/Source/PresetAsset/Private/AssetDefinition_PresetAsset.cpp
new file mode 100644
--- /dev/null
+++ b/Engine/Plugins/Experimental/ToolPresets/Source/PresetAsset/Private/AssetDefinition_PresetAsset.cpp
@@ -0,0 +1,7 @@
+// Copyright Epic Games, Inc. All Rights Reserved.
+
+#include "AssetDefinition_PresetAsset.h"
+
+#define LOCTEXT_NAMESPACE "AssetDefinition_PresetAsset"
+
+#undef LOCTEXT_NAMESPACE
\ No newline at end of file
diff --git a/Engine/Plugins/Experimental/ToolPresets/Source/PresetAsset/Private/AssetDefinition_PresetAsset.h b/Engine/Plugins/Experimental/ToolPresets/Source/PresetAsset/Private/AssetDefinition_PresetAsset.h
new file mode 100644
--- /dev/null
+++ b/Engine/Plugins/Experimental/ToolPresets/Source/PresetAsset/Private/AssetDefinition_PresetAsset.h
@@ -0,0 +1,30 @@
+// Copyright Epic Games, Inc. All Rights Reserved.
+
+#pragma once
+
+#include "PresetAsset.h"
+#include "AssetDefinitionDefault.h"
+
+#include "AssetDefinition_PresetAsset.generated.h"
+
+enum class EAssetCommandResult : uint8;
+struct FAssetCategoryPath;
+struct FAssetOpenArgs;
+
+UCLASS()
+class UAssetDefinition_PresetAsset : public UAssetDefinitionDefault
+{
+	GENERATED_BODY()
+
+public:
+	// UAssetDefinition Begin
+	virtual FText GetAssetDisplayName() const override { return NSLOCTEXT("AssetTypeActions", "AssetTypeActions_PresetAsset", "PresetAsset"); }
+	virtual FLinearColor GetAssetColor() const override { return FLinearColor(FColor(175, 0, 128)); }
+	virtual TSoftClassPtr<UObject> GetAssetClass() const override { return UPresetAsset::StaticClass(); }
+	virtual TConstArrayView<FAssetCategoryPath> GetAssetCategories() const override
+	{
+		return TConstArrayView<FAssetCategoryPath>();
+	}
+	virtual FText GetObjectDisplayNameText(UObject* Object) const override { return FText::FromString(TEXT("UPresetAsset")); }
+	// UAssetDefinition End
+};
\ No newline at end of file
diff --git a/Engine/Plugins/Experimental/ToolPresets/Source/PresetAsset/Private/PresetAssetModule.cpp b/Engine/Plugins/Experimental/ToolPresets/Source/PresetAsset/Private/PresetAssetModule.cpp
new file mode 100644
--- /dev/null
+++ b/Engine/Plugins/Experimental/ToolPresets/Source/PresetAsset/Private/PresetAssetModule.cpp
@@ -0,0 +1,64 @@
+// Copyright Epic Games, Inc. All Rights Reserved.
+
+#include "Modules/ModuleInterface.h"
+#include "Modules/ModuleManager.h"
+#include "PresetAsset.h"
+#include "PropertyEditorModule.h"
+
+#define LOCTEXT_NAMESPACE "FPresetAssetModule"
+
+
+/**
+ * Implements the PresetAsset module.
+ */
+class FPresetAssetModule
+	: public IModuleInterface
+{
+public:
+
+	//~ IModuleInterface interface
+
+	virtual void StartupModule() override;
+	virtual void ShutdownModule() override;
+	virtual void OnPostEngineInit();
+	virtual bool SupportsDynamicReloading() override
+	{
+		return true;
+	}
+
+protected:
+
+	TArray<FName> ClassesToUnregisterOnShutdown;
+};
+
+
+// This function will execute after your module is loaded into memory; the exact timing is specified in the .uplugin file per-module
+void FPresetAssetModule::StartupModule()
+{
+	FCoreDelegates::OnPostEngineInit.AddRaw(this, &FPresetAssetModule::OnPostEngineInit);
+}
+
+// This function may be called during shutdown to clean up your module.  For modules that support dynamic reloading,
+// we call this function before unloading the module.
+void FPresetAssetModule::ShutdownModule()
+{
+	// Unregister customizations
+	FPropertyEditorModule* PropertyEditorModule = FModuleManager::GetModulePtr<FPropertyEditorModule>("PropertyEditor");
+	if (PropertyEditorModule)
+	{
+		for (FName ClassName : ClassesToUnregisterOnShutdown)
+		{
+			PropertyEditorModule->UnregisterCustomClassLayout(ClassName);
+		}
+	}
+}
+
+void FPresetAssetModule::OnPostEngineInit()
+{
+	ClassesToUnregisterOnShutdown.Reset();
+
+	FPropertyEditorModule& PropertyModule = FModuleManager::LoadModuleChecked<FPropertyEditorModule>("PropertyEditor");
+}
+
+
+IMPLEMENT_MODULE(FPresetAssetModule, PresetAsset);
\ No newline at end of file
diff --git a/Engine/Plugins/Experimental/ToolPresets/Source/PresetAsset/Public/PresetAsset.h b/Engine/Plugins/Experimental/ToolPresets/Source/PresetAsset/Public/PresetAsset.h
new file mode 100644
--- /dev/null
+++ b/Engine/Plugins/Experimental/ToolPresets/Source/PresetAsset/Public/PresetAsset.h
@@ -0,0 +1,107 @@
+// Copyright Epic Games, Inc. All Rights Reserved.
+
+#pragma once
+
+//#include "Internationalization/Text.h"
+#include "UObject/Object.h"
+#include "UObject/ObjectMacros.h"
+#include "Factories/Factory.h"
+#include "AssetToolsModule.h"
+
+#include "PresetAsset.generated.h"
+
+#define LOCTEXT_NAMESPACE "PresetAsset"
+
+USTRUCT(BlueprintType)
+struct FPropertyStore
+{
+	GENERATED_BODY()
+public:
+	UPROPERTY(VisibleAnywhere, Category = "PresetAsset")
+	TArray< TObjectPtr<UObject> > Properties;
+};
+
+USTRUCT(BlueprintType)
+struct FNamedPropertyStore
+{
+	GENERATED_BODY()
+public:
+	UPROPERTY(EditAnywhere, Category = "PresetAsset")
+	TMap< FString, FPropertyStore> Store;
+};
+
+/**
+ * Implements an asset that can be used to store tool settings as a named preset
+ */
+UCLASS(BlueprintType, hidecategories=(Object))
+class PRESETASSET_API UPresetAsset
+	: public UObject
+{
+	GENERATED_BODY()
+
+public:
+
+	// TODO: Currently there are no helper methods within this class, simply providing
+	// raw access to the underlying arrays and maps. This is intentional.
+	// Until the design of the preset concept is more firmly decided, it seems like a waste to
+	// implement a bunch of methods that we don't know if we actually want/need at the end.
+	// Once we are satisifed with the data structure, planned accessors and mutators will include
+	// support for adding, removing, renaming, saving and retrieving presets.
+
+	UPROPERTY(VisibleAnywhere, Category = "PresetAsset")
+	TMap<FString, FNamedPropertyStore > StoredProperties;
+
+};
+
+
+UCLASS(hidecategories = Object)
+class PRESETASSET_API UPresetAssetFactory : public UFactory
+{
+	GENERATED_BODY()
+
+public:
+
+	UPresetAssetFactory()
+	{
+		SupportedClass = UPresetAsset::StaticClass();
+
+		bCreateNew = true;
+		bEditAfterNew = true;
+	}
+
+	FText GetDisplayName() const override
+	{
+		return LOCTEXT("DisplayName", "Tool Preset");
+	}
+
+	FText GetToolTip() const override
+	{
+		return LOCTEXT("Tooltip", "Tool Presets capture the state of tool settings for later reloading.");
+	}
+
+	UObject* FactoryCreateNew(UClass* Class, UObject* InParent, FName Name, EObjectFlags Flags, UObject* Context, FFeedbackContext* Warn) override
+	{
+		UPresetAsset* Preset = nullptr;
+		if (ensure(SupportedClass == Class))
+		{
+			ensure(0 != (RF_Public & Flags));
+			Preset = NewObject<UPresetAsset>(InParent, Class, Name, Flags);
+		}
+		return Preset;
+	}
+
+	uint32 GetMenuCategories() const override
+	{
+		IAssetTools& AssetTools = FModuleManager::LoadModuleChecked<FAssetToolsModule>("AssetTools").Get();
+		return AssetTools.RegisterAdvancedAssetCategory("Presets", LOCTEXT("AssetCategoryName", "Presets"));
+	}
+
+	FString GetDefaultNewAssetName() const override
+	{
+		return TEXT("Tool Preset");
+	}
+
+	bool ShouldShowInNewMenu() const override { return true; }
+};
+
+#undef LOCTEXT_NAMESPACE
\ No newline at end of file
diff --git a/Engine/Plugins/Experimental/ToolPresets/ToolPresets.uplugin b/Engine/Plugins/Experimental/ToolPresets/ToolPresets.uplugin
new file mode 100644
--- /dev/null
+++ b/Engine/Plugins/Experimental/ToolPresets/ToolPresets.uplugin
@@ -0,0 +1,25 @@
+{
+	"FileVersion": 3,
+	"Version": 1,
+	"VersionName": "0.1",
+	"FriendlyName": "Tool Presets",
+	"Description": "Adds support for saving and loading tool settings as presets.",
+	"Category": "Editor",
+	"CreatedBy": "Epic Games Inc",
+	"CreatedByURL": "http://epicgames.com",
+	"DocsURL": "",
+	"MarketplaceURL": "",
+	"SupportURL": "",
+	"EnabledByDefault": true,
+	"CanContainContent": false,
+	"IsBetaVersion": false,
+	"Installed": false,
+	"IsExperimentalVersion": true,
+	"Modules": [
+		{
+			"Name": "PresetAsset",
+			"Type": "Editor",
+			"LoadingPhase": "Default"
+		}
+	]
+}
\ No newline at end of file
