ModelingMode: Address static analysis error with the new preset code

diff --git a/Engine/Plugins/Editor/ModelingToolsEditorMode/Source/ModelingToolsEditorMode/Private/ModelingToolsEditorModeToolkit.cpp b/Engine/Plugins/Editor/ModelingToolsEditorMode/Source/ModelingToolsEditorMode/Private/ModelingToolsEditorModeToolkit.cpp
--- a/Engine/Plugins/Editor/ModelingToolsEditorMode/Source/ModelingToolsEditorMode/Private/ModelingToolsEditorModeToolkit.cpp
+++ b/Engine/Plugins/Editor/ModelingToolsEditorMode/Source/ModelingToolsEditorMode/Private/ModelingToolsEditorModeToolkit.cpp
@@ -748,18 +748,18 @@ void FModelingToolsEditorModeToolkit::SaveActivePreset()
 
 		for (UObject* PropertySet : PropertySets)
 		{
-			UObject* StoredPropertySet = NewObject<UObject>(&Preset, PropertySet->GetClass());
-			StoredPropertySet->ClearFlags(EObjectFlags::RF_Transient);
+			UObject* StoredPropertySet = NewObject<UObject>(&Preset, PropertySet->GetClass());			
 			if (StoredPropertySet != nullptr)
 			{
+				StoredPropertySet->ClearFlags(EObjectFlags::RF_Transient);
 				for (FProperty* Prop : TFieldRange<FProperty>(PropertySet->GetClass()))
 				{
 					void* SrcValue = Prop->ContainerPtrToValuePtr<void>(PropertySet);
 					void* DestValue = Prop->ContainerPtrToValuePtr<void>(StoredPropertySet);
 					Prop->CopySingleValue(DestValue, SrcValue);
 				}
-			}
-			PropertyStore.Add(StoredPropertySet);
+				PropertyStore.Add(StoredPropertySet);
+			}			
 		}	
 
 		Preset.MarkPackageDirty();
