ModelingMode: Fix crash errors due to preset combo box being accessed as a nullptr

diff --git a/Engine/Plugins/Editor/ModelingToolsEditorMode/Source/ModelingToolsEditorMode/Private/ModelingToolsEditorModeToolkit.cpp b/Engine/Plugins/Editor/ModelingToolsEditorMode/Source/ModelingToolsEditorMode/Private/ModelingToolsEditorModeToolkit.cpp
--- a/Engine/Plugins/Editor/ModelingToolsEditorMode/Source/ModelingToolsEditorMode/Private/ModelingToolsEditorModeToolkit.cpp
+++ b/Engine/Plugins/Editor/ModelingToolsEditorMode/Source/ModelingToolsEditorMode/Private/ModelingToolsEditorModeToolkit.cpp
@@ -653,8 +653,11 @@ void FModelingToolsEditorModeToolkit::ClearPresetComboList()
 {
 	ActiveNamedPreset.Reset();
 	AvailablePresetsForTool.Empty();
-	PresetComboBox->RefreshOptions();
-	PresetComboBox->ClearSelection();
+	if (PresetComboBox)
+	{
+		PresetComboBox->RefreshOptions();
+		PresetComboBox->ClearSelection();
+	}
 }
 
 void FModelingToolsEditorModeToolkit::HandlePresetAssetChanged(const FAssetData& InAssetData)
