rename MshMrg to Merge, MshDup to Dupe, and Merge to Union. Clean up tooltips.

diff --git a/Engine/Plugins/Editor/ModelingToolsEditorMode/Source/ModelingToolsEditorMode/Private/ModelingToolsManagerActions.cpp b/Engine/Plugins/Editor/ModelingToolsEditorMode/Source/ModelingToolsEditorMode/Private/ModelingToolsManagerActions.cpp
--- a/Engine/Plugins/Editor/ModelingToolsEditorMode/Source/ModelingToolsEditorMode/Private/ModelingToolsManagerActions.cpp
+++ b/Engine/Plugins/Editor/ModelingToolsEditorMode/Source/ModelingToolsEditorMode/Private/ModelingToolsManagerActions.cpp
@@ -63,109 +63,109 @@ void FModelingToolsManagerCommands::RegisterCommands()
 	REGISTER_MODELING_TOOL_COMMAND(BeginAddStairsPrimitiveTool, "Stairs", "Create new stairs objects");
 
 	// Create
-	REGISTER_MODELING_TOOL_COMMAND(BeginDrawPolygonTool, "PolyExt", "Draw and extrude polygons to create new objects");
-	REGISTER_MODELING_TOOL_COMMAND(BeginDrawPolyPathTool, "PathExt", "Draw and extrude PolyPaths to create new objects");
-	REGISTER_MODELING_TOOL_COMMAND(BeginDrawAndRevolveTool, "PathRev", "Draw and revolve PolyPaths to create new objects");
-	REGISTER_MODELING_TOOL_COMMAND(BeginRevolveBoundaryTool, "BdryRev", "Revolve mesh boundary loops to create new objects");
-	REGISTER_MODELING_TOOL_COMMAND(BeginCombineMeshesTool, "MshMrg", "Merge multiple meshes to create new objects");
-	REGISTER_MODELING_TOOL_COMMAND(BeginDuplicateMeshesTool, "MshDup", "Duplicate single meshes to create new objects");
+	REGISTER_MODELING_TOOL_COMMAND(BeginDrawPolygonTool, "PolyExt", "Draw and extrude 2D Polygons to create new objects");
+	REGISTER_MODELING_TOOL_COMMAND(BeginDrawPolyPathTool, "PathExt", "Draw and extrude 2D Paths to create new objects");
+	REGISTER_MODELING_TOOL_COMMAND(BeginDrawAndRevolveTool, "PathRev", "Draw and revolve 2D Paths to create new objects");
+	REGISTER_MODELING_TOOL_COMMAND(BeginRevolveBoundaryTool, "BdryRev", "Revolve Mesh boundary loops to create new objects");
+	REGISTER_MODELING_TOOL_COMMAND(BeginCombineMeshesTool, "Merge", "Merge multiple Meshes to create new objects");
+	REGISTER_MODELING_TOOL_COMMAND(BeginDuplicateMeshesTool, "Dupe", "Duplicate single Meshes to create new objects");
 
 	// PolyModel
-	REGISTER_MODELING_TOOL_COMMAND(BeginPolyEditTool, "PolyEd", "Edit meshes via PolyGroups");
-	REGISTER_MODELING_TOOL_COMMAND(BeginPolyDeformTool, "PolyDef", "Deform meshes via PolyGroups");
-	REGISTER_MODELING_TOOL_COMMAND(BeginCubeGridTool, "CubeGr", "Create block out meshes using a repositionable grid");
-	REGISTER_MODELING_TOOL_COMMAND(BeginMeshBooleanTool, "MshBool", "Apply Boolean operations to mesh pairs");
-	REGISTER_MODELING_TOOL_COMMAND(BeginCutMeshWithMeshTool, "MshCut", "Split one mesh into parts using a second mesh");
-	REGISTER_MODELING_TOOL_COMMAND(BeginSubdividePolyTool, "SubDiv", "Subdivide mesh via PolyGroups or triangles");
+	REGISTER_MODELING_TOOL_COMMAND(BeginPolyEditTool, "PolyEd", "Edit Meshes via PolyGroups");
+	REGISTER_MODELING_TOOL_COMMAND(BeginPolyDeformTool, "PolyDef", "Deform Meshes via PolyGroups");
+	REGISTER_MODELING_TOOL_COMMAND(BeginCubeGridTool, "CubeGr", "Create block out Meshes using a repositionable grid");
+	REGISTER_MODELING_TOOL_COMMAND(BeginMeshBooleanTool, "MshBool", "Apply Boolean operations to the two selected Meshes");
+	REGISTER_MODELING_TOOL_COMMAND(BeginCutMeshWithMeshTool, "MshCut", "Split one Mesh into parts using a second Mesh");
+	REGISTER_MODELING_TOOL_COMMAND(BeginSubdividePolyTool, "SubDiv", "Subdivide the selected Mesh via PolyGroups or triangles");
 
 	// TriModel
-	REGISTER_MODELING_TOOL_COMMAND(BeginMeshSelectionTool, "TriSel", "Select and edit mesh triangles");
-	REGISTER_MODELING_TOOL_COMMAND(BeginTriEditTool, "TriEd", "Edit mesh via triangles");
-	REGISTER_MODELING_TOOL_COMMAND(BeginHoleFillTool, "HFill", "Fill holes in mesh");
-	REGISTER_MODELING_TOOL_COMMAND(BeginMirrorTool, "Mirror", "Mirror selected meshes");
-	REGISTER_MODELING_TOOL_COMMAND(BeginPlaneCutTool, "PlnCut", "Cut selected meshes with plane");
-	REGISTER_MODELING_TOOL_COMMAND(BeginPolygonCutTool, "PolyCut", "Cut mesh with extruded polygon");
-	REGISTER_MODELING_TOOL_COMMAND(BeginMeshTrimTool, "Trim", "Trim/cut selected mesh with second mesh");
+	REGISTER_MODELING_TOOL_COMMAND(BeginMeshSelectionTool, "TriSel", "Select and edit Mesh triangles with a brush interface");
+	REGISTER_MODELING_TOOL_COMMAND(BeginTriEditTool, "TriEd", "Select and Edit the Mesh vertices, edges, and triangles");
+	REGISTER_MODELING_TOOL_COMMAND(BeginHoleFillTool, "HFill", "Fill holes in the selected Mesh");
+	REGISTER_MODELING_TOOL_COMMAND(BeginMirrorTool, "Mirror", "Mirror the selected Meshes");
+	REGISTER_MODELING_TOOL_COMMAND(BeginPlaneCutTool, "PlnCut", "Cut the selected Meshes with a 3D plane");
+	REGISTER_MODELING_TOOL_COMMAND(BeginPolygonCutTool, "PolyCut", "Cut the selected Mesh with an extruded polygon");
+	REGISTER_MODELING_TOOL_COMMAND(BeginMeshTrimTool, "Trim", "Trim/Cut the selected Mesh with the second selected Mesh");
 
 	// Deform
 	REGISTER_MODELING_TOOL_COMMAND(BeginSculptMeshTool, "VSclpt", "Vertex sculpting");
 	REGISTER_MODELING_TOOL_COMMAND(BeginRemeshSculptMeshTool, "DSclpt", "Dynamic mesh sculpting");
-	REGISTER_MODELING_TOOL_COMMAND(BeginSmoothMeshTool, "Smooth", "Smooth mesh surface");
-	REGISTER_MODELING_TOOL_COMMAND(BeginOffsetMeshTool, "Offset", "Offset mesh surface");
-	REGISTER_MODELING_TOOL_COMMAND(BeginMeshSpaceDeformerTool, "Warp", "Reshape mesh using space deformers");
-	REGISTER_MODELING_TOOL_COMMAND(BeginLatticeDeformerTool, "Lattice", "Deform mesh with 3D lattice/grid");
-	REGISTER_MODELING_TOOL_COMMAND(BeginDisplaceMeshTool, "Displce", "Displace mesh surface with optional subdivision");
+	REGISTER_MODELING_TOOL_COMMAND(BeginSmoothMeshTool, "Smooth", "Smooth the shape of the selected Mesh");
+	REGISTER_MODELING_TOOL_COMMAND(BeginOffsetMeshTool, "Offset", "Offset the surface of the selected Mesh");
+	REGISTER_MODELING_TOOL_COMMAND(BeginMeshSpaceDeformerTool, "Warp", "Reshape the selected Mesh using space deformers");
+	REGISTER_MODELING_TOOL_COMMAND(BeginLatticeDeformerTool, "Lattice", "Deform the selected Mesh using a 3D lattice/grid");
+	REGISTER_MODELING_TOOL_COMMAND(BeginDisplaceMeshTool, "Displce", "Tessellate and Displace the selected Mesh");
 
 	// Transform
-	REGISTER_MODELING_TOOL_COMMAND(BeginTransformMeshesTool, "XForm", "Transform selected meshes");
-	REGISTER_MODELING_TOOL_COMMAND(BeginAlignObjectsTool, "Align", "Align objects");
-	REGISTER_MODELING_TOOL_COMMAND(BeginEditPivotTool, "Pivot", "Edit mesh pivots");
+	REGISTER_MODELING_TOOL_COMMAND(BeginTransformMeshesTool, "XForm", "Transform the selected Meshes");
+	REGISTER_MODELING_TOOL_COMMAND(BeginAlignObjectsTool, "Align", "Align the selected Objects");
+	REGISTER_MODELING_TOOL_COMMAND(BeginEditPivotTool, "Pivot", "Edit the pivot points of the selected Meshes");
 	REGISTER_MODELING_TOOL_COMMAND(BeginAddPivotActorTool, "PivotAct", "Add actor to act as a pivot for child components");
-	REGISTER_MODELING_TOOL_COMMAND(BeginBakeTransformTool, "BakeRS", "Bake rotation and scale into mesh asset");
-	REGISTER_MODELING_TOOL_COMMAND(BeginTransferMeshTool, "Transfer", "Transfer meshes");
-	REGISTER_MODELING_TOOL_COMMAND(BeginConvertMeshesTool, "Convert", "Convert meshes");
-	REGISTER_MODELING_TOOL_COMMAND(BeginSplitMeshesTool, "Split", "Split meshes");
-	REGISTER_MODELING_TOOL_COMMAND(BeginPatternTool, "Pattern", "Create patterns of meshes");
+	REGISTER_MODELING_TOOL_COMMAND(BeginBakeTransformTool, "BakeRS", "Bake rotation and scale into the selected Meshes");
+	REGISTER_MODELING_TOOL_COMMAND(BeginTransferMeshTool, "Transfer", "Copy the first selected Mesh to the second selected Mesh");
+	REGISTER_MODELING_TOOL_COMMAND(BeginConvertMeshesTool, "Convert", "Convert the selected Meshes to a different type of Mesh Object");
+	REGISTER_MODELING_TOOL_COMMAND(BeginSplitMeshesTool, "Split", "Split the selected Meshes into separate parts based on connectivity");
+	REGISTER_MODELING_TOOL_COMMAND(BeginPatternTool, "Pattern", "Create patterns of Meshes");
 
 	// MeshOps
-	REGISTER_MODELING_TOOL_COMMAND(BeginSimplifyMeshTool, "Simplfy", "Simplify mesh");
-	REGISTER_MODELING_TOOL_COMMAND(BeginRemeshMeshTool, "Remesh", "Re-triangulate mesh");
-	REGISTER_MODELING_TOOL_COMMAND(BeginWeldEdgesTool, "Weld", "Weld overlapping mesh edges");
-	REGISTER_MODELING_TOOL_COMMAND(BeginRemoveOccludedTrianglesTool, "Jacket", "Remove hidden triangles from selected meshes");
-	REGISTER_MODELING_TOOL_COMMAND(BeginSelfUnionTool, "Merge", "Self-union selected meshes to resolve self-intersections");
-	REGISTER_MODELING_TOOL_COMMAND(BeginProjectToTargetTool, "Project", "Map/re-mesh onto target mesh (second selection)");
+	REGISTER_MODELING_TOOL_COMMAND(BeginSimplifyMeshTool, "Simplfy", "Simplify the selected Mesh");
+	REGISTER_MODELING_TOOL_COMMAND(BeginRemeshMeshTool, "Remesh", "Re-triangulate the selected Mesh");
+	REGISTER_MODELING_TOOL_COMMAND(BeginWeldEdgesTool, "Weld", "Weld overlapping Mesh edges");
+	REGISTER_MODELING_TOOL_COMMAND(BeginRemoveOccludedTrianglesTool, "Jacket", "Remove hidden triangles from the selected Meshes");
+	REGISTER_MODELING_TOOL_COMMAND(BeginSelfUnionTool, "Union", "Boolean Union the selected Meshes, including Self-Union to resolve self-intersections");
+	REGISTER_MODELING_TOOL_COMMAND(BeginProjectToTargetTool, "Project", "Map/re-mesh the first selected Mesh onto the second selected Mesh");
 
 	// VoxOps
-	REGISTER_MODELING_TOOL_COMMAND(BeginVoxelSolidifyTool, "VoxWrap", "Wrap selected meshes using voxels");
-	REGISTER_MODELING_TOOL_COMMAND(BeginVoxelBlendTool, "VoxBlnd", "Blend selected meshes using voxels");
-	REGISTER_MODELING_TOOL_COMMAND(BeginVoxelMorphologyTool, "VoxMrph", "Offset/inset selected meshes using voxels");
+	REGISTER_MODELING_TOOL_COMMAND(BeginVoxelSolidifyTool, "VoxWrap", "Wrap the selected Meshes using voxels");
+	REGISTER_MODELING_TOOL_COMMAND(BeginVoxelBlendTool, "VoxBlnd", "Blend the selected Meshes using voxels");
+	REGISTER_MODELING_TOOL_COMMAND(BeginVoxelMorphologyTool, "VoxMrph", "Offset/Inset the selected Meshes using voxels");
 #if WITH_PROXYLOD
 	// The ProxyLOD plugin is currently only available on Windows. Without it, the following tools do not work as expected.
-	REGISTER_MODELING_TOOL_COMMAND(BeginVoxelBooleanTool, "VoxBool", "Boolean selected meshes using voxels");
-	REGISTER_MODELING_TOOL_COMMAND(BeginVoxelMergeTool, "VoxMrg", "Merge selected meshes using voxels");
+	REGISTER_MODELING_TOOL_COMMAND(BeginVoxelBooleanTool, "VoxBool", "Boolean the selected Meshes using voxels");
+	REGISTER_MODELING_TOOL_COMMAND(BeginVoxelMergeTool, "VoxMrg", "Merge the selected Meshes using voxels");
 #endif	// WITH_PROXYLOD
 
 	// Attributes
-	REGISTER_MODELING_TOOL_COMMAND(BeginMeshInspectorTool, "Inspct", "Inspect mesh attributes");
-	REGISTER_MODELING_TOOL_COMMAND(BeginEditNormalsTool, "Nrmls", "Update mesh normals");
-	REGISTER_MODELING_TOOL_COMMAND(BeginEditTangentsTool, "Tngnts", "Update mesh tangents");
-	REGISTER_MODELING_TOOL_COMMAND(BeginAttributeEditorTool, "AttrEd", "Edit/configure mesh attributes");
+	REGISTER_MODELING_TOOL_COMMAND(BeginMeshInspectorTool, "Inspct", "Inspect Mesh attributes");
+	REGISTER_MODELING_TOOL_COMMAND(BeginEditNormalsTool, "Nrmls", "Recompute or Repair Normals");
+	REGISTER_MODELING_TOOL_COMMAND(BeginEditTangentsTool, "Tngnts", "Recompute Tangents");
+	REGISTER_MODELING_TOOL_COMMAND(BeginAttributeEditorTool, "AttrEd", "Edit/configure Mesh attributes");
 	REGISTER_MODELING_TOOL_COMMAND(BeginPolyGroupsTool, "GrpGen", "Generate new PolyGroups");
 	REGISTER_MODELING_TOOL_COMMAND(BeginMeshGroupPaintTool, "GrpPnt", "Paint new PolyGroups");
 	REGISTER_MODELING_TOOL_COMMAND(BeginMeshAttributePaintTool, "MapPnt", "Paint attribute maps");
 	REGISTER_MODELING_TOOL_COMMAND(BeginEditMeshMaterialsTool, "MatEd", "Assign materials to selected triangles");
 
 	// UVs
-	REGISTER_MODELING_TOOL_COMMAND(BeginGlobalUVGenerateTool, "AutoUV", "Automatically unwrap and pack UVs for mesh");
-	REGISTER_MODELING_TOOL_COMMAND(BeginGroupUVGenerateTool, "Unwrap", "Perform UV unwrapping for mesh");
-	REGISTER_MODELING_TOOL_COMMAND(BeginUVProjectionTool, "Project", "Set UVs from projection");
-	REGISTER_MODELING_TOOL_COMMAND(BeginUVSeamEditTool, "SeamEd", "Add UV seams to mesh");
+	REGISTER_MODELING_TOOL_COMMAND(BeginGlobalUVGenerateTool, "AutoUV", "Automatically unwrap and pack UVs");
+	REGISTER_MODELING_TOOL_COMMAND(BeginGroupUVGenerateTool, "Unwrap", "Recompute UVs for existing UV islands or Polygroups");
+	REGISTER_MODELING_TOOL_COMMAND(BeginUVProjectionTool, "Project", "Compute UVs via projecting to simple shapes");
+	REGISTER_MODELING_TOOL_COMMAND(BeginUVSeamEditTool, "SeamEd", "Add UV seams");
 	REGISTER_MODELING_TOOL_COMMAND(BeginTransformUVIslandsTool, "XFormUV", "Transform UV islands in UV space");
-	REGISTER_MODELING_TOOL_COMMAND(BeginUVLayoutTool, "Layout", "Transform and pack existing UVs");
+	REGISTER_MODELING_TOOL_COMMAND(BeginUVLayoutTool, "Layout", "Transform and Repack existing UVs");
 	// This is done directly, not with the REGISTER_ macro, since we don't want it added to the tool list or use a toggle button
 	UI_COMMAND(LaunchUVEditor, "UVEditor", "Launch UV asset editor", EUserInterfaceActionType::Button, FInputChord());
 
 	// Baking
-	REGISTER_MODELING_TOOL_COMMAND(BeginBakeMeshAttributeMapsTool, "BakeTx", "Bake textures for a target mesh");
-	REGISTER_MODELING_TOOL_COMMAND(BeginBakeMultiMeshAttributeMapsTool, "BakeAll", "Bake textures for a target mesh from multiple source meshes");
-	REGISTER_MODELING_TOOL_COMMAND(BeginBakeRenderCaptureTool, "BakeRC", "Bake textures for a target mesh from multiple source meshes via virtual photo/render capture");
-	REGISTER_MODELING_TOOL_COMMAND(BeginBakeMeshAttributeVertexTool, "BakeVtx", "Bake vertex colors for a target mesh");
+	REGISTER_MODELING_TOOL_COMMAND(BeginBakeMeshAttributeMapsTool, "BakeTx", "Bake textures for a target Mesh");
+	REGISTER_MODELING_TOOL_COMMAND(BeginBakeMultiMeshAttributeMapsTool, "BakeAll", "Bake textures for a target Mesh from multiple source Meshes");
+	REGISTER_MODELING_TOOL_COMMAND(BeginBakeRenderCaptureTool, "BakeRC", "Bake renders into new textures for a target Mesh from multiple source Meshes");
+	REGISTER_MODELING_TOOL_COMMAND(BeginBakeMeshAttributeVertexTool, "BakeVtx", "Bake vertex colors for a target Mesh");
 
 	// Volumes
-	REGISTER_MODELING_TOOL_COMMAND(BeginVolumeToMeshTool, "Vol2Msh", "Convert volume to new mesh asset");
-	REGISTER_MODELING_TOOL_COMMAND(BeginMeshToVolumeTool, "Msh2Vol", "Convert mesh to volume");
+	REGISTER_MODELING_TOOL_COMMAND(BeginVolumeToMeshTool, "Vol2Msh", "Convert a Volume to a new Mesh Object");
+	REGISTER_MODELING_TOOL_COMMAND(BeginMeshToVolumeTool, "Msh2Vol", "Convert a Mesh to a Volume");
 	if (!Settings->InRestrictiveMode())
 	{
-		REGISTER_MODELING_TOOL_COMMAND(BeginBspConversionTool, "BSPConv", "Convert BSP to static mesh asset");
+		REGISTER_MODELING_TOOL_COMMAND(BeginBspConversionTool, "BSPConv", "Convert BSP to a new Mesh Object");
 	}
-	REGISTER_MODELING_TOOL_COMMAND(BeginPhysicsInspectorTool, "PInspct", "Inspect physics geometry for selected meshes");
-	REGISTER_MODELING_TOOL_COMMAND(BeginSetCollisionGeometryTool, "Msh2Coll", "Convert selected meshes to Simple Collision Geometry (for last selected)");
-	REGISTER_MODELING_TOOL_COMMAND(BeginExtractCollisionGeometryTool, "Coll2Msh", "Convert Simple Collision Geometry to new mesh asset");
+	REGISTER_MODELING_TOOL_COMMAND(BeginPhysicsInspectorTool, "PInspct", "Inspect the physics/collision geometry for selected Meshes");
+	REGISTER_MODELING_TOOL_COMMAND(BeginSetCollisionGeometryTool, "Msh2Coll", "Convert selected Meshes to Simple Collision Geometry (for last selected)");
+	REGISTER_MODELING_TOOL_COMMAND(BeginExtractCollisionGeometryTool, "Coll2Msh", "Convert Simple Collision Geometry to a new Mesh Object");
 
 	// LODs
-	REGISTER_MODELING_TOOL_COMMAND(BeginLODManagerTool, "LODMgr", "Static mesh asset LOD manager");
-	REGISTER_MODELING_TOOL_COMMAND(BeginGenerateStaticMeshLODAssetTool, "AutoLOD", "Generate static mesh LOD asset");
-	REGISTER_MODELING_TOOL_COMMAND(BeginISMEditorTool, "ISMEd", "Edit Instaces in InstancedStaticMeshComponents");
+	REGISTER_MODELING_TOOL_COMMAND(BeginLODManagerTool, "LODMgr", "Inspect the LODs of a Static Mesh Asset");
+	REGISTER_MODELING_TOOL_COMMAND(BeginGenerateStaticMeshLODAssetTool, "AutoLOD", "Automatically generate a simplified LOD with baked Textures/Materials for a Mesh");
+	REGISTER_MODELING_TOOL_COMMAND(BeginISMEditorTool, "ISMEd", "Edit the Instances of Instanced Static Mesh Components");
 
 	REGISTER_MODELING_TOOL_COMMAND(BeginAddPatchTool, "Patch", "Add Patch");
 	REGISTER_MODELING_TOOL_COMMAND(BeginShapeSprayTool, "Spray", "Shape Spray");
@@ -175,17 +175,17 @@ void FModelingToolsManagerCommands::RegisterCommands()
 	UI_COMMAND(BeginSkinWeightsPaintTool, "SkinWts", "Start the Paint Skin Weights Tool", EUserInterfaceActionType::Button, FInputChord());
 	UI_COMMAND(BeginSkinWeightsBindingTool, "SkinBind", "Start the Skin Weights Binding Tool", EUserInterfaceActionType::Button, FInputChord());
 
-	REGISTER_MODELING_TOOL_COMMAND(BeginPolyModelTool_Inset, "Inset", "PolyGroup Inset Tool");
-	REGISTER_MODELING_TOOL_COMMAND(BeginPolyModelTool_Outset, "Outset", "PolyGroup Outset Tool");
-	REGISTER_MODELING_TOOL_COMMAND(BeginPolyModelTool_CutFaces, "Cut", "PolyGroup Cut Faces Tool");
-	REGISTER_MODELING_TOOL_COMMAND(BeginPolyModelTool_InsertEdgeLoop, "ELoop", "PolyGroup Insert Edge Loop Tool");
-	REGISTER_MODELING_TOOL_COMMAND(BeginPolyModelTool_PushPull, "PushPull", "PolyGroup Push Pull Faces Tool");
-	REGISTER_MODELING_TOOL_COMMAND(BeginPolyModelTool_Bevel, "Bevel", "PolyGroup Bevel Tool");
+	REGISTER_MODELING_TOOL_COMMAND(BeginPolyModelTool_Inset, "Inset", "Inset the current Selected Faces");
+	REGISTER_MODELING_TOOL_COMMAND(BeginPolyModelTool_Outset, "Outset", "Outset the current Selected Faces");
+	REGISTER_MODELING_TOOL_COMMAND(BeginPolyModelTool_CutFaces, "Cut", "Cut the current Selected Faces");
+	REGISTER_MODELING_TOOL_COMMAND(BeginPolyModelTool_InsertEdgeLoop, "ELoop", "Insert Edge Loops into the Selected Mesh");
+	REGISTER_MODELING_TOOL_COMMAND(BeginPolyModelTool_PushPull, "PushPull", "Push/Pull the current Selected Faces");
+	REGISTER_MODELING_TOOL_COMMAND(BeginPolyModelTool_Bevel, "Bevel", "Bevel the current Mesh Selection (Edges or Faces)");
 
-	REGISTER_MODELING_TOOL_COMMAND(BeginPolyModelTool_PolyEd, "PolyEd", "Edit meshes via PolyGroups");
-	REGISTER_MODELING_TOOL_COMMAND(BeginPolyModelTool_TriSel, "TriSel", "Select and edit mesh triangles");
+	REGISTER_MODELING_TOOL_COMMAND(BeginPolyModelTool_PolyEd, "PolyEd", "Select / Edit the current Mesh via PolyGroups");
+	REGISTER_MODELING_TOOL_COMMAND(BeginPolyModelTool_TriSel, "TriSel", "Select / Edit the current Mesh via Triangles");
 
-	REGISTER_MODELING_TOOL_COMMAND(MeshSelectionModeAction_NoSelection, "None", "Disable Geometry Selection");
+	REGISTER_MODELING_TOOL_COMMAND(MeshSelectionModeAction_NoSelection, "None", "Disable Mesh Selection");
 	REGISTER_MODELING_TOOL_COMMAND(MeshSelectionModeAction_MeshTriangles, "Tris", "Select Mesh Triangles");
 	REGISTER_MODELING_TOOL_COMMAND(MeshSelectionModeAction_MeshVertices, "Verts", "Select Mesh Vertices");
 	REGISTER_MODELING_TOOL_COMMAND(MeshSelectionModeAction_MeshEdges, "Edges", "Select Mesh Edges");
@@ -193,29 +193,29 @@ void FModelingToolsManagerCommands::RegisterCommands()
 	REGISTER_MODELING_TOOL_COMMAND(MeshSelectionModeAction_GroupCorners, "Corners", "Select Mesh Polygroup Corners/Vertices");
 	REGISTER_MODELING_TOOL_COMMAND(MeshSelectionModeAction_GroupEdges, "Borders", "Select Mesh Polygroup Borders/Edges");
 
-	REGISTER_MODELING_ACTION_COMMAND(BeginSelectionAction_Delete, "Delete", "Delete Selection");
-	REGISTER_MODELING_ACTION_COMMAND(BeginSelectionAction_Disconnect, "Discon", "Disconnect Selection");
-	REGISTER_MODELING_ACTION_COMMAND(BeginSelectionAction_Retriangulate, "Clean", "Retriangulate Selection");
+	REGISTER_MODELING_ACTION_COMMAND(BeginSelectionAction_Delete, "Delete", "Delete the current Mesh Selection");
+	REGISTER_MODELING_ACTION_COMMAND(BeginSelectionAction_Disconnect, "Discon", "Disconnect the current Mesh Selection");
+	REGISTER_MODELING_ACTION_COMMAND(BeginSelectionAction_Retriangulate, "Clean", "Retriangulate the current Mesh or Mesh Selection");
 
-	REGISTER_MODELING_TOOL_COMMAND(BeginSelectionAction_Extrude, "Extrude", "Extrude Selection");
-	REGISTER_MODELING_TOOL_COMMAND(BeginSelectionAction_Offset, "Offset", "Offset Selection");
+	REGISTER_MODELING_TOOL_COMMAND(BeginSelectionAction_Extrude, "Extrude", "Extrude the current Mesh Selection");
+	REGISTER_MODELING_TOOL_COMMAND(BeginSelectionAction_Offset, "Offset", "Offset the current Mesh Selection");
 
 	REGISTER_MODELING_ACTION_COMMAND(BeginSelectionAction_SelectAll, "Select All", "Select All Elements");
 	REGISTER_MODELING_ACTION_COMMAND(BeginSelectionAction_ExpandToConnected, "Expand To Connected", "Expand Selection to Geometrically-Connected Elements");
 	REGISTER_MODELING_ACTION_COMMAND(BeginSelectionAction_Invert, "Invert Selection", "Invert the current Selection");
 	REGISTER_MODELING_ACTION_COMMAND(BeginSelectionAction_InvertConnected, "Invert Connected", "Invert the current Selection to Geometrically-Connected Elements");
-	REGISTER_MODELING_ACTION_COMMAND(BeginSelectionAction_Expand, "Expand Selection", "Expand Selection");
-	REGISTER_MODELING_ACTION_COMMAND(BeginSelectionAction_Contract, "Contract Selection", "Contract Selection");
+	REGISTER_MODELING_ACTION_COMMAND(BeginSelectionAction_Expand, "Expand Selection", "Expand the current Selection by a ring of elements");
+	REGISTER_MODELING_ACTION_COMMAND(BeginSelectionAction_Contract, "Contract Selection", "Contract the current Selection by a ring of elements");
 
 
 
-	UI_COMMAND(AcceptActiveTool, "Accept", "Accept the active tool", EUserInterfaceActionType::Button, FInputChord());
-	UI_COMMAND(CancelActiveTool, "Cancel", "Cancel the active tool", EUserInterfaceActionType::Button, FInputChord());
-	UI_COMMAND(CompleteActiveTool, "Done", "Complete the active tool", EUserInterfaceActionType::Button, FInputChord());
+	UI_COMMAND(AcceptActiveTool, "Accept", "Accept the active Tool", EUserInterfaceActionType::Button, FInputChord());
+	UI_COMMAND(CancelActiveTool, "Cancel", "Cancel the active Tool", EUserInterfaceActionType::Button, FInputChord());
+	UI_COMMAND(CompleteActiveTool, "Done", "Complete the active Tool", EUserInterfaceActionType::Button, FInputChord());
 
 	// Note that passing a chord into one of these calls hooks the key press to the respective action. 
-	UI_COMMAND(AcceptOrCompleteActiveTool, "Accept or Complete", "Accept or complete the active tool", EUserInterfaceActionType::Button, FInputChord(EKeys::Enter));
-	UI_COMMAND(CancelOrCompleteActiveTool, "Cancel or Complete", "Cancel or complete the active tool", EUserInterfaceActionType::Button, FInputChord(EKeys::Escape));
+	UI_COMMAND(AcceptOrCompleteActiveTool, "Accept or Complete", "Accept or Complete the active Tool", EUserInterfaceActionType::Button, FInputChord(EKeys::Enter));
+	UI_COMMAND(CancelOrCompleteActiveTool, "Cancel or Complete", "Cancel or Complete the active Tool", EUserInterfaceActionType::Button, FInputChord(EKeys::Escape));
 	
 #undef REGISTER_MODELING_TOOL_COMMAND
 }
