Modeling Mode: Preset UI second version

* Moves the experimental preset controls to the top of the modeling mode tool's panel.
* Refactors the preset list and create preset widgets into a new, single menu.
* Enables support for multiple preset assets to be pulled from as preset libraries for a tool
* Renames preset asset classes to better reflect their intended use in the Interactive Tools Framework

https://drive.google.com/file/d/1E4kqJJu277u6Ca3LlAa0TTTPyEXwO7ZM/view?usp=sharing
https://drive.google.com/file/d/1cKzBT9PJDAL6C_SGGwnBOyRkr4dhLl52/view?usp=sharing
https://drive.google.com/file/d/1kdid8J2W2QUwyGlQm7nTAHe3TXCEscAX/view?usp=sharing
https://drive.google.com/file/d/1x9-4hjj7apMaRoD90ZqUpNWUeab3X3PS/view?usp=sharing

diff --git a/Engine/Plugins/Editor/ModelingToolsEditorMode/Source/ModelingToolsEditorMode/Private/ModelingToolsEditorModeToolkit.cpp b/Engine/Plugins/Editor/ModelingToolsEditorMode/Source/ModelingToolsEditorMode/Private/ModelingToolsEditorModeToolkit.cpp
--- a/Engine/Plugins/Editor/ModelingToolsEditorMode/Source/ModelingToolsEditorMode/Private/ModelingToolsEditorModeToolkit.cpp
+++ b/Engine/Plugins/Editor/ModelingToolsEditorMode/Source/ModelingToolsEditorMode/Private/ModelingToolsEditorModeToolkit.cpp
@@ -68,16 +68,16 @@ static TAutoConsoleVariable<int32> CVarEnableToolPresets(
 
 namespace FModelingToolsEditorModeToolkitLocals
 {
-	typedef TFunction<void(UPresetAsset& Preset, UInteractiveTool& Tool)> PresetAndToolFunc;
+	typedef TFunction<void(UInteractiveToolsPresetCollectionAsset& Preset, UInteractiveTool& Tool)> PresetAndToolFunc;
 
 	void ExecuteWithPresetAndTool(UEdMode& EdMode, EToolSide ToolSide, FSoftObjectPath& PresetPath, PresetAndToolFunc Function)
 	{
-		UPresetAsset* Preset = nullptr;		
+		UInteractiveToolsPresetCollectionAsset* Preset = nullptr;		
 		UInteractiveTool* Tool = EdMode.GetToolManager()->GetActiveTool(EToolSide::Left);
 
 		if (PresetPath.IsAsset())
 		{
-			Preset = Cast<UPresetAsset>(PresetPath.TryLoad());
+			Preset = Cast<UInteractiveToolsPresetCollectionAsset>(PresetPath.TryLoad());
 		}
 		if (!Preset || !Tool)
 		{
@@ -85,13 +85,28 @@ namespace FModelingToolsEditorModeToolkitLocals
 		}
 		Function(*Preset, *Tool);
 	}
-}
 
+	void ExecuteWithPresetAndTool(UEdMode& EdMode, EToolSide ToolSide, UInteractiveToolsPresetCollectionAsset& Preset, PresetAndToolFunc Function)
+	{
+		UInteractiveTool* Tool = EdMode.GetToolManager()->GetActiveTool(EToolSide::Left);
+
+		if (!Tool)
+		{
+			return;
+		}
+		Function(Preset, *Tool);
+	}
+}
 
 FModelingToolsEditorModeToolkit::FModelingToolsEditorModeToolkit()
 {
+	PresetSettings = NewObject<UPresetSettingsProperties>();
 }
 
+FString FModelingToolsEditorModeToolkit::GetReferencerName() const
+{
+	return TEXT("FModelingToolsEditorModeToolkit");
+}
 
 FModelingToolsEditorModeToolkit::~FModelingToolsEditorModeToolkit()
 {
@@ -114,6 +129,8 @@ FModelingToolsEditorModeToolkit::~FModelingToolsEditorModeToolkit()
 	Settings->OnModified.Remove(AssetSettingsModifiedHandle);
 	GetScriptableEditorMode()->GetInteractiveToolsContext(EToolsContextScope::EdMode)->OnToolNotificationMessage.RemoveAll(this);
 	GetScriptableEditorMode()->GetInteractiveToolsContext(EToolsContextScope::EdMode)->OnToolWarningMessage.RemoveAll(this);
+
+	PresetSettings = nullptr;
 }
 
 
@@ -174,6 +191,10 @@ void FModelingToolsEditorModeToolkit::Init(const TSharedPtr<IToolkitHost>& InitT
 	ToolWarningArea->SetText(FText::GetEmpty());
 
 	// add the various sections to the mode toolbox
+	ToolkitWidgetVBox->AddSlot().AutoHeight().HAlign(HAlign_Fill).Padding(5)
+		[
+			MakePresetPanel()->AsShared()
+		];
 	ToolkitWidgetVBox->AddSlot().AutoHeight().HAlign(HAlign_Fill).Padding(5)
 		[
 			ModeWarningArea->AsShared()
@@ -222,6 +243,10 @@ void FModelingToolsEditorModeToolkit::Init(const TSharedPtr<IToolkitHost>& InitT
 	}
 }
 
+void FModelingToolsEditorModeToolkit::AddReferencedObjects(FReferenceCollector& Collector)
+{
+	Collector.AddReferencedObject(PresetSettings);
+}
 
 
 void FModelingToolsEditorModeToolkit::MakeToolShutdownOverlayWidget()
@@ -456,73 +481,6 @@ TSharedPtr<SWidget> FModelingToolsEditorModeToolkit::MakeAssetConfigPanel()
 		];
 	}
 
-	bool bEnableToolPresets = (CVarEnableToolPresets.GetValueOnGameThread() > 0);
-
-	if(bEnableToolPresets)
-	{
-		BuildPresetComboList();
-		Content->AddSlot().HAlign(HAlign_Fill).VAlign(VAlign_Fill).AutoHeight()[
-			SNew(SHorizontalBox)
-			+ SHorizontalBox::Slot().HAlign(HAlign_Right).Padding(0, 2, 2, 2).AutoWidth()
-			[
-				SNew(STextBlock)
-				.Text(LOCTEXT("PresetLabel", "Preset"))
-			]
-			+ SHorizontalBox::Slot().HAlign(HAlign_Fill).Padding(0, 0, 2, 0).FillWidth(1.0f)
-			[
-				SNew(SVerticalBox)			
-				+ SVerticalBox::Slot().AutoHeight()
-				[
-					SNew(SHorizontalBox)
-					+ SHorizontalBox::Slot().HAlign(HAlign_Fill).Padding(0, 2, 2, 2).FillWidth(1.0f)
-					[
-						SNew(SObjectPropertyEntryBox)
-						.AllowedClass(UPresetAsset::StaticClass())
-						.ObjectPath_Lambda([this]() {return FModelingToolsEditorModeToolkit::GetCurrentPresetPath(); })
-						.OnShouldFilterAsset(this, &FModelingToolsEditorModeToolkit::HandleFilterPresetAsset)
-						.OnObjectChanged(this, &FModelingToolsEditorModeToolkit::HandlePresetAssetChanged)
-					]
-				]
-				+ SVerticalBox::Slot().Padding(0, 2, 2, 2).AutoHeight()
-				[
-					SNew(SHorizontalBox)
-					+ SHorizontalBox::Slot().HAlign(HAlign_Fill).Padding(0, 2, 2, 2).FillWidth(1.0f)
-					[
-						SAssignNew(PresetComboBox, SEditableComboBox<TSharedPtr<FString>>)
-						.OptionsSource(&AvailablePresetsForTool)
-						.OnGenerateWidget(this, &FModelingToolsEditorModeToolkit::MakePresetComboWidget)
-						.OnSelectionChanged(this, &FModelingToolsEditorModeToolkit::OnPresetChanged)
-						//.OnComboBoxOpening(this, &FModelingToolsEditorModeToolkit::BuildPresetComboList)
-						.OnGetEditableText(this, &FModelingToolsEditorModeToolkit::GetRawPresetString)
-						.OnSelectionRenamed(this, &FModelingToolsEditorModeToolkit::OnPresetRenamed)
-						.OnAddClicked(this, &FModelingToolsEditorModeToolkit::OnAddPreset)
-						.OnRemoveClicked(this, &FModelingToolsEditorModeToolkit::OnRemovePreset)
-						.IsEnabled(this, &FModelingToolsEditorModeToolkit::IsPresetEnabled)
-						.ContentPadding(2)
-						.Content()
-						[
-							SNew(STextBlock)
-							.Text(this, &FModelingToolsEditorModeToolkit::GetPresetComboBoxContent)
-							.ToolTipText(this, &FModelingToolsEditorModeToolkit::GetPresetComboBoxContent)
-						]
-					]
-					+ SHorizontalBox::Slot().HAlign(HAlign_Right).Padding(0, 2, 2, 2).AutoWidth()
-					[
-						SNew(SSimpleButton)
-						.OnClicked_Lambda([this]() { LoadActivePreset(); return FReply::Handled(); })
-					.Icon(FAppStyle::Get().GetBrush("Icons.Load"))
-					]
-					+ SHorizontalBox::Slot().HAlign(HAlign_Right).Padding(0, 2, 2, 2).AutoWidth()
-					[
-						SNew(SSimpleButton)
-						.OnClicked_Lambda([this]() { SaveActivePreset(); return FReply::Handled(); })
-					.Icon(FAppStyle::Get().GetBrush("Icons.Save"))
-					]
-				]
-			]
-		];
-	}
-
 	TSharedPtr<SExpandableArea> AssetConfigPanel = SNew(SExpandableArea)
 		.HeaderPadding(FMargin(2.0f))
 		.Padding(FMargin(2.f))
@@ -553,186 +511,265 @@ TSharedRef<SWidget> FModelingToolsEditorModeToolkit::MakePresetComboWidget(TShar
 		.Text(FText::FromString(*InItem));
 }
 
+TSharedPtr<SWidget> FModelingToolsEditorModeToolkit::MakePresetPanel()
+{
+	const TSharedPtr<SVerticalBox> Content = SNew(SVerticalBox);
+	UModelingToolsEditorModeSettings* Settings = GetMutableDefault<UModelingToolsEditorModeSettings>();
+
+	bool bEnableToolPresets = (CVarEnableToolPresets.GetValueOnGameThread() > 0);
+	if (!bEnableToolPresets || Settings->InRestrictiveMode())
+	{
+		return SNew(SVerticalBox);
+	}
+
+	const TSharedPtr<SHorizontalBox> NewContent = SNew(SHorizontalBox);
+	
+	NewContent->AddSlot().HAlign(HAlign_Right)
+	[
+		SNew(SComboButton)
+		.ComboButtonStyle(&FAppStyle::Get().GetWidgetStyle<FComboButtonStyle>("SimpleComboButton"))
+		.OnGetMenuContent(this, &FModelingToolsEditorModeToolkit::GetPresetCreateButtonContent)
+		.HasDownArrow(true)
+		.ButtonContent()
+		[
+			SNew(SHorizontalBox)
+			+ SHorizontalBox::Slot()
+			.AutoWidth()
+			.VAlign(VAlign_Center)
+			.Padding(4.0, 0.0f)
+			[
+				SNew(SImage)
+				.ColorAndOpacity(FSlateColor::UseForeground())
+				.Image(FAppStyle::Get().GetBrush("Icons.Plus"))
+			]
+			+ SHorizontalBox::Slot()
+			.AutoWidth()
+			.VAlign(VAlign_Center)
+			.Padding(4.0, 0.0f)
+			[
+				SNew(STextBlock)
+				.Text(LOCTEXT("ModelingPresetPanelHeader", "Presets"))
+				.Justification(ETextJustify::Center)
+				.Font(FAppStyle::Get().GetFontStyle("EditorModesPanel.CategoryFontStyle"))
+			]
+
+		]
+	];
+
+	NewContent->AddSlot().HAlign(HAlign_Right).AutoWidth()
+		[
+			SNew(SComboButton)
+			.ComboButtonStyle(&FAppStyle::Get().GetWidgetStyle<FComboButtonStyle>("SimpleComboButton"))
+		.OnGetMenuContent(this, &FModelingToolsEditorModeToolkit::GetPresetSettingsButtonContent)
+		.OnMenuOpenChanged(this, &FModelingToolsEditorModeToolkit::RebuildPresetListForTool)
+		.HasDownArrow(false)
+		.ButtonContent()
+		[
+			SNew(SHorizontalBox)
+			+ SHorizontalBox::Slot()
+		.AutoWidth()
+		.VAlign(VAlign_Center)
+		.Padding(4.0, 0.0f)
+		[
+			SNew(SImage)
+			.ColorAndOpacity(FSlateColor::UseForeground())
+		.Image(FAppStyle::Get().GetBrush("Icons.Settings"))
+		]
+		]
+		];
+
+	TSharedPtr<SHorizontalBox> AssetConfigPanel = SNew(SHorizontalBox)
+		+ SHorizontalBox::Slot()
+		[
+			NewContent->AsShared()
+		];
+
+	return AssetConfigPanel;
+}
+
 bool FModelingToolsEditorModeToolkit::IsPresetEnabled() const
 {
 	return CurrentPreset.IsAsset();
 }
 
-void FModelingToolsEditorModeToolkit::OnPresetRenamed(const FText& NewPresetName, ETextCommit::Type CommitType)
+TSharedRef<SWidget> FModelingToolsEditorModeToolkit::GetPresetSettingsButtonContent()
 {
-	FModelingToolsEditorModeToolkitLocals::ExecuteWithPresetAndTool(*OwningEditorMode, EToolSide::Left, CurrentPreset,
-		[this, NewPresetName] (UPresetAsset& Preset, UInteractiveTool& Tool) {
-		FPropertyStore PresetValuesToRename;
-		if (!Preset.StoredProperties.Contains(Tool.GetClass()->GetName()))
-		{
-			return;
-		}
-		if (Preset.StoredProperties[Tool.GetClass()->GetName()].Store.Contains(*GetPresetString()))
-		{
-			Preset.StoredProperties[Tool.GetClass()->GetName()].Store.RemoveAndCopyValue(*GetPresetString(), PresetValuesToRename);
-			Preset.StoredProperties[Tool.GetClass()->GetName()].Store.Add(NewPresetName.ToString(), PresetValuesToRename);
-			BuildPresetComboList();
-			ActiveNamedPreset = *AvailablePresetsForTool.FindByPredicate([this, NewPresetName](TSharedPtr<FString>& Value) { return Value->Equals(NewPresetName.ToString()); });
+    const TSharedPtr<SVerticalBox> Content = SNew(SVerticalBox);
 
-			Preset.MarkPackageDirty();
-		}
-	});
-}
+	FPropertyEditorModule& PropertyEditorModule = FModuleManager::GetModuleChecked<FPropertyEditorModule>("PropertyEditor");
 
-FReply FModelingToolsEditorModeToolkit::OnAddPreset()
-{
-	FModelingToolsEditorModeToolkitLocals::ExecuteWithPresetAndTool(*OwningEditorMode, EToolSide::Left, CurrentPreset,
-		[this] (UPresetAsset& Preset, UInteractiveTool& Tool) {
+	TSharedPtr<IDetailsView> PresetSettingsDetailsView;
 
-		FString NewPresetName = FString::Printf(TEXT("NewPreset_%d"), Preset.StoredProperties.FindOrAdd(Tool.GetClass()->GetName()).Store.Num());
-		FPropertyStore PresetValuesToCreate;
-		Preset.StoredProperties[Tool.GetClass()->GetName()].Store.Add(NewPresetName, PresetValuesToCreate);
-		BuildPresetComboList();
-		ActiveNamedPreset = *AvailablePresetsForTool.FindByPredicate([this, NewPresetName](TSharedPtr<FString>& Value) { return Value->Equals(NewPresetName); });
+	FDetailsViewArgs DetailsViewArgs;
+	DetailsViewArgs.bAllowSearch = false;
+	DetailsViewArgs.NameAreaSettings = FDetailsViewArgs::HideNameArea;
+	DetailsViewArgs.bHideSelectionTip = true;
+	DetailsViewArgs.DefaultsOnlyVisibility = EEditDefaultsOnlyNodeVisibility::Automatic;
+	DetailsViewArgs.bShowOptions = false;
+	DetailsViewArgs.bAllowMultipleTopLevelObjects = true;
 
-		Preset.MarkPackageDirty();
-	});
-	return FReply::Handled();
-}
+	PresetSettingsDetailsView = PropertyEditorModule.CreateDetailView(DetailsViewArgs);
+	
+	PresetSettingsDetailsView->SetObject(PresetSettings);
 
-FReply FModelingToolsEditorModeToolkit::OnRemovePreset()
-{
-	FModelingToolsEditorModeToolkitLocals::ExecuteWithPresetAndTool(*OwningEditorMode, EToolSide::Left, CurrentPreset,
-		[this] (UPresetAsset& Preset, UInteractiveTool& Tool) {
-		if (!Preset.StoredProperties.Contains(Tool.GetClass()->GetName()))
-		{
-			return;
-		}
-		FPropertyStore PresetValuesToRename;
-		if (Preset.StoredProperties[Tool.GetClass()->GetName()].Store.Contains(*GetPresetString()))
-		{
-			Preset.StoredProperties[Tool.GetClass()->GetName()].Store.RemoveAndCopyValue(*GetPresetString(), PresetValuesToRename);
-			ActiveNamedPreset.Reset();
-			BuildPresetComboList();
+	Content->AddSlot()
+	[
+		PresetSettingsDetailsView->AsShared()
+	];
 
-			Preset.MarkPackageDirty();
-		}
-	});
-	return FReply::Handled();
+	return Content->AsShared();
 }
 
-void FModelingToolsEditorModeToolkit::OnPresetChanged(TSharedPtr<FString> NewSelection, ESelectInfo::Type SelectInfo)
+TSharedRef<SWidget> FModelingToolsEditorModeToolkit::GetPresetCreateButtonContent()
 {
-	ActiveNamedPreset = NewSelection;
-}
+
+	constexpr bool bShouldCloseMenuAfterSelection = true;
+	FMenuBuilder MenuBuilder(bShouldCloseMenuAfterSelection, nullptr);
+
+	constexpr bool bNoIndent = false;
+	constexpr bool bSearchable = false;
+	static const FName NoExtensionHook = NAME_None;
 
 
-FText FModelingToolsEditorModeToolkit::GetPresetComboBoxContent() const
-{
-	TSharedPtr<FString> CurrentPresetValue = GetPresetString();
-	if (CurrentPresetValue)
 	{
-		return FText::FromString(*CurrentPresetValue.Get());
+		MenuBuilder.BeginSection(NoExtensionHook, LOCTEXT("ListPresetMenu", "Presets for Tool"));
+		FMenuEntryParams MenuEntryParams;
+
+		for (TSharedPtr<FToolPresetOption> ToolPresetOption : AvailablePresetsForTool)
+		{
+			FUIAction ApplyPresetAction;
+			ApplyPresetAction.ExecuteAction = FExecuteAction::CreateLambda([this, ToolPresetOption]()
+				{
+					LoadPresetFromCollection(ToolPresetOption->PresetName, ToolPresetOption->PresetCollection);
+				});
+
+			ApplyPresetAction.CanExecuteAction = FCanExecuteAction::CreateLambda([this]()
+				{
+					return this->OwningEditorMode->GetToolManager()->GetActiveTool(EToolSide::Left) != nullptr;
+				});
+
+			MenuBuilder.AddMenuEntry(
+				FText::FromString(ToolPresetOption->PresetLabel),
+				FText::FromString(ToolPresetOption->PresetTooltip),
+				ToolPresetOption->PresetIcon,
+				ApplyPresetAction);
+		}
+
+		MenuBuilder.EndSection();
 	}
-	else
+
 	{
-		return FText::GetEmpty();
-	}
-}
+		MenuBuilder.BeginSection(NoExtensionHook, LOCTEXT("CreatePresetMenu", "Create New Preset"));
+		FMenuEntryParams MenuEntryParams;
 
+		TSharedRef<SWidget> PresetName = SNew(SEditableTextBox)
+			.OnTextCommitted_Lambda([this](const FText& NewLabel, ETextCommit::Type&) { NewPresetLabel = NewLabel.ToString(); });
 
-TSharedPtr<FString> FModelingToolsEditorModeToolkit::GetPresetString() const
-{
-	return ActiveNamedPreset;
-}
+		MenuBuilder.AddWidget(PresetName, LOCTEXT("NewPresetLabelLabel", "Label"), bNoIndent, bSearchable, LOCTEXT("NewPresetLabelTooltip", "A label that identifies the preset in the preset dropdown."));
 
-FString FModelingToolsEditorModeToolkit::GetRawPresetString() const
-{
-	return GetPresetString() ? *GetPresetString() : TEXT("");
-}
+		TSharedRef<SObjectPropertyEntryBox> PresetCollection = SNew(SObjectPropertyEntryBox)
+			.AllowedClass(UInteractiveToolsPresetCollectionAsset::StaticClass())
+			.ObjectPath_Lambda([this]() {return FModelingToolsEditorModeToolkit::GetCurrentPresetPath(); })
+			.OnShouldFilterAsset(this, &FModelingToolsEditorModeToolkit::HandleFilterPresetAsset)
+			.OnObjectChanged(this, &FModelingToolsEditorModeToolkit::HandlePresetAssetChanged);
 
-void FModelingToolsEditorModeToolkit::BuildPresetComboList()
-{
-	FModelingToolsEditorModeToolkitLocals::ExecuteWithPresetAndTool(*OwningEditorMode, EToolSide::Left, CurrentPreset,
-		[this] (UPresetAsset& Preset, UInteractiveTool& Tool) {
+		MenuBuilder.AddWidget(PresetCollection, LOCTEXT("NewPresetCollectionLabel", "Collection"), bNoIndent, bSearchable, LOCTEXT("NewPresetCollectionTooltip", "The preset collection asset that should store this preset."));
 
-		AvailablePresetsForTool.Empty();
-		if (!Preset.StoredProperties.Contains(Tool.GetClass()->GetName()))
-		{
-			PresetComboBox->RefreshOptions();
-			PresetComboBox->ClearSelection();
-			return;
-		}
-		TArray<FString> PresetsForTool;
-		Preset.StoredProperties[Tool.GetClass()->GetName()].Store.GenerateKeyArray(PresetsForTool);
-		AvailablePresetsForTool.Reserve(PresetsForTool.Num());
-		for (FString& PresetEntry : PresetsForTool)
-		{
-			AvailablePresetsForTool.Add(MakeShared<FString>(PresetEntry));
-		}
-		PresetComboBox->RefreshOptions();
-		PresetComboBox->ClearSelection();
-	});
+		TSharedRef<SEditableTextBox> PresetToolTip = SNew(SEditableTextBox)
+			.OnTextCommitted_Lambda([this](const FText& NewToolTip, ETextCommit::Type&) { NewPresetTooltip = NewToolTip.ToString(); });
+
+		MenuBuilder.AddWidget(PresetToolTip, LOCTEXT("NewPresetTooltipLabel", "Tooltip"), bNoIndent, bSearchable, LOCTEXT("NewPresetTiooltipTooltip", "A descriptive tooltip that describes the purpose of this preset."));
+
+		FUIAction AddPresetAction;
+		AddPresetAction.ExecuteAction = FExecuteAction::CreateLambda([this, &PresetName, &MenuBuilder]()
+			{
+				CreateNewPresetInCollection(NewPresetLabel,
+					CurrentPreset,
+					NewPresetTooltip,
+					NewPresetIcon);
+			});
+
+		AddPresetAction.CanExecuteAction = FCanExecuteAction::CreateLambda([this]()
+			{
+				return this->OwningEditorMode->GetToolManager()->GetActiveTool(EToolSide::Left) != nullptr && CurrentPreset.IsValid();
+			});
+
+		MenuBuilder.AddMenuEntry(LOCTEXT("CreatePresetEntry", "Create Preset"),
+			LOCTEXT("CreatePresetToolTip", "Add a new preset for this tool from current settings."),
+			FSlateIcon(FAppStyle::Get().GetStyleSetName(), "Icons.Plus"),
+			AddPresetAction);
+
+		MenuBuilder.EndSection();
+	}
+
+	return MenuBuilder.MakeWidget();
 }
 
 void FModelingToolsEditorModeToolkit::ClearPresetComboList()
 {
-	ActiveNamedPreset.Reset();
 	AvailablePresetsForTool.Empty();
-	if (PresetComboBox)
-	{
-		PresetComboBox->RefreshOptions();
-		PresetComboBox->ClearSelection();
-	}
 }
 
-void FModelingToolsEditorModeToolkit::HandlePresetAssetChanged(const FAssetData& InAssetData)
+void FModelingToolsEditorModeToolkit::RebuildPresetListForTool(bool bSettingsOpened)
 {
-	CurrentPreset.SetPath(InAssetData.GetObjectPathString());
-	if (CurrentPreset.IsAsset())
+	if (bSettingsOpened)
 	{
-		BuildPresetComboList();
+		return;
 	}
-	else
+
+	AvailablePresetsForTool.Empty();
+	for (FSoftObjectPath& PresetCollection : PresetSettings->ActivePresetCollectionsPaths)
 	{
-		CurrentPreset.Reset();
+		FModelingToolsEditorModeToolkitLocals::ExecuteWithPresetAndTool(*OwningEditorMode, EToolSide::Left, PresetCollection,
+			[this, PresetCollection](UInteractiveToolsPresetCollectionAsset& Preset, UInteractiveTool& Tool) {
+		
+				if (!Preset.PerToolPresets.Contains(Tool.GetClass()->GetName()))
+				{
+					return;
+				}
+				TArray<FString> PresetsForTool;
+				Preset.PerToolPresets[Tool.GetClass()->GetName()].NamedPresets.GenerateKeyArray(PresetsForTool);
+				AvailablePresetsForTool.Reserve(PresetsForTool.Num());
+				for (FString& PresetEntry : PresetsForTool)
+				{
+					TSharedPtr<FToolPresetOption> NewOption = MakeShared<FToolPresetOption>();
+					NewOption->PresetLabel = Preset.PerToolPresets[Tool.GetClass()->GetName()].NamedPresets[PresetEntry].Label;
+					NewOption->PresetTooltip = Preset.PerToolPresets[Tool.GetClass()->GetName()].NamedPresets[PresetEntry].Tooltip;
+					//NewOption->PresetIcon = Preset.PerToolPresets[Tool.GetClass()->GetName()].Store[PresetEntry].Icon;
+					NewOption->PresetName = PresetEntry;
+					NewOption->PresetCollection = PresetCollection;
+
+					AvailablePresetsForTool.Add(NewOption);
+				}
+			});
 	}
 }
 
+void FModelingToolsEditorModeToolkit::HandlePresetAssetChanged(const FAssetData& InAssetData)
+{
+	CurrentPreset.SetPath(InAssetData.GetObjectPathString());
+}
+
 bool FModelingToolsEditorModeToolkit::HandleFilterPresetAsset(const FAssetData& InAssetData)
 {
-#if 1
 	return false;
-#else 
-    // This was the "old" approach, where only presets valid for the current tool were shown. But this precluded selecting
-    // other assets and *adding* the current tool to them, so the filtering was removed. However, the logic may prove useful,
-    // so this is just keeping it around for a little while longer just in case.
-
-	UPresetAsset* PresetAsset = Cast<UPresetAsset>(InAssetData.GetAsset());
-	if (PresetAsset) 
-	{
-		bool bMatchingToolName = false;
-		if (OwningEditorMode->GetToolManager()->HasActiveTool(EToolSide::Left))
-		{
-			bMatchingToolName |= PresetAsset->StoredProperties.Contains(OwningEditorMode->GetToolManager()->GetActiveTool(EToolSide::Left)->GetClass()->GetName());
-		}
-		if (OwningEditorMode->GetToolManager()->HasActiveTool(EToolSide::Right))
-		{
-			bMatchingToolName |= PresetAsset->StoredProperties.Contains(OwningEditorMode->GetToolManager()->GetActiveTool(EToolSide::Right)->GetClass()->GetName());
-		}
-		return !bMatchingToolName;
-	}
-
-	return true;
-#endif
 }
 
-void FModelingToolsEditorModeToolkit::LoadActivePreset()
+void FModelingToolsEditorModeToolkit::LoadPresetFromCollection(const FString& PresetName, FSoftObjectPath CollectionPath)
 {
-	FModelingToolsEditorModeToolkitLocals::ExecuteWithPresetAndTool(*OwningEditorMode, EToolSide::Left, CurrentPreset,
-	[this](UPresetAsset& Preset, UInteractiveTool& Tool) {
-		if (!Preset.StoredProperties.Contains(Tool.GetClass()->GetName()))
+	FModelingToolsEditorModeToolkitLocals::ExecuteWithPresetAndTool(*OwningEditorMode, EToolSide::Left, CollectionPath,
+	[this, PresetName](UInteractiveToolsPresetCollectionAsset& Preset, UInteractiveTool& Tool) {
+		TArray<UObject*> PropertySets = Tool.GetToolProperties();
+		if (!Preset.PerToolPresets.Contains(Tool.GetClass()->GetName()))			
+		{
+			return;
+		}
+		if(!Preset.PerToolPresets[Tool.GetClass()->GetName()].NamedPresets.Contains(PresetName))
 		{
 			return;
 		}
 
-		TArray<UObject*> PropertySets = Tool.GetToolProperties();
-		TArray<TObjectPtr<UObject>>& PropertyStore = Preset.StoredProperties[Tool.GetClass()->GetName()].Store.FindOrAdd(*ActiveNamedPreset).Properties;
+		TArray<TObjectPtr<UObject>>& PropertyStore = Preset.PerToolPresets[Tool.GetClass()->GetName()].NamedPresets.Find(*PresetName)->Properties;
 
 		for (UObject* PropertySet : PropertySets)
 		{
@@ -761,33 +798,82 @@ void FModelingToolsEditorModeToolkit::LoadActivePreset()
 	});
 }
 
-void FModelingToolsEditorModeToolkit::SaveActivePreset()
+void FModelingToolsEditorModeToolkit::CreateNewPresetInCollection(const FString& PresetLabel, FSoftObjectPath CollectionPath, const FString& ToolTip, FSlateIcon Icon)
 {
+	FModelingToolsEditorModeToolkitLocals::ExecuteWithPresetAndTool(*OwningEditorMode, EToolSide::Left, CurrentPreset,
+	[this, PresetLabel, ToolTip, Icon, CollectionPath](UInteractiveToolsPresetCollectionAsset& Preset, UInteractiveTool& Tool) {
+
+		TArray<UObject*> PropertySets = Tool.GetToolProperties();
+		FString NewPresetName = FString::Printf(TEXT("NewPreset_%d"), Preset.PerToolPresets.FindOrAdd(Tool.GetClass()->GetName()).NamedPresets.Num());
+		FInteractiveToolPresetDefintion PresetValuesToCreate;
+		PresetValuesToCreate.Label = PresetLabel;
+		PresetValuesToCreate.Tooltip = ToolTip;
+		//PresetValuesToCreate.Icon = Icon;
+
+		TArray<TObjectPtr<UObject>>& PropertyStore = PresetValuesToCreate.Properties;
+		PropertyStore.Empty();
+
+		for (UObject* PropertySet : PropertySets)
+		{
+			UObject* StoredPropertySet = NewObject<UObject>(&Preset, PropertySet->GetClass());
+			StoredPropertySet->ClearFlags(EObjectFlags::RF_Transient);
+			if (StoredPropertySet != nullptr)
+			{
+				for (FProperty* Prop : TFieldRange<FProperty>(PropertySet->GetClass()))
+				{
+					void* SrcValue = Prop->ContainerPtrToValuePtr<void>(PropertySet);
+					void* DestValue = Prop->ContainerPtrToValuePtr<void>(StoredPropertySet);
+					Prop->CopySingleValue(DestValue, SrcValue);
+				}
+			}
+			PropertyStore.Add(StoredPropertySet);
+		}
+
+		Preset.PerToolPresets[Tool.GetClass()->GetName()].NamedPresets.Add(NewPresetName, PresetValuesToCreate);
+		Preset.MarkPackageDirty();
+
+		// Finally add this to the current tool's preset list, since we know it should be there.
+		TSharedPtr<FToolPresetOption> NewOption = MakeShared<FToolPresetOption>();
+		NewOption->PresetLabel = PresetLabel;
+		NewOption->PresetTooltip = ToolTip;
+		NewOption->PresetName = NewPresetName;
+		NewOption->PresetCollection = CollectionPath;
+
+		AvailablePresetsForTool.Add(NewOption);
+		
+	});
+}
 
+void FModelingToolsEditorModeToolkit::SaveActivePreset()
+{
+	
+	// Saving presets has been disabled for the moment until the 3rd UI generation is done, but we'll want this logic then.
+#if 0
 	FModelingToolsEditorModeToolkitLocals::ExecuteWithPresetAndTool(*OwningEditorMode, EToolSide::Left, CurrentPreset,
-	[this](UPresetAsset& Preset, UInteractiveTool& Tool) {
+	[this](UInteractiveToolsPresetCollectionAsset& Preset, UInteractiveTool& Tool) {
 		TArray<UObject*> PropertySets = Tool.GetToolProperties();
-		TArray<TObjectPtr<UObject>>& PropertyStore = Preset.StoredProperties.FindOrAdd(Tool.GetClass()->GetName()).Store.FindOrAdd(*ActiveNamedPreset).Properties;
+		TArray<TObjectPtr<UObject>>& PropertyStore = Preset.PerToolPresets.FindOrAdd(Tool.GetClass()->GetName()).NamedPresets.FindOrAdd(*ActiveNamedPreset).Properties;
 		PropertyStore.Empty();
 
 		for (UObject* PropertySet : PropertySets)
 		{
-			UObject* StoredPropertySet = NewObject<UObject>(&Preset, PropertySet->GetClass());			
+			UObject* StoredPropertySet = NewObject<UObject>(&Preset, PropertySet->GetClass());
+			StoredPropertySet->ClearFlags(EObjectFlags::RF_Transient);
 			if (StoredPropertySet != nullptr)
 			{
-				StoredPropertySet->ClearFlags(EObjectFlags::RF_Transient);
 				for (FProperty* Prop : TFieldRange<FProperty>(PropertySet->GetClass()))
 				{
 					void* SrcValue = Prop->ContainerPtrToValuePtr<void>(PropertySet);
 					void* DestValue = Prop->ContainerPtrToValuePtr<void>(StoredPropertySet);
 					Prop->CopySingleValue(DestValue, SrcValue);
 				}
-				PropertyStore.Add(StoredPropertySet);
-			}			
+			}
+			PropertyStore.Add(StoredPropertySet);
 		}	
 
 		Preset.MarkPackageDirty();
 	});
+#endif
 }
 
 void FModelingToolsEditorModeToolkit::InitializeAfterModeSetup()
@@ -1488,8 +1574,7 @@ void FModelingToolsEditorModeToolkit::OnToolStarted(UInteractiveToolManager* Man
 			Viewport->Invalidate();
 		}
 	}
-
-	BuildPresetComboList();
+	RebuildPresetListForTool(false);
 }
 
 void FModelingToolsEditorModeToolkit::OnToolEnded(UInteractiveToolManager* Manager, UInteractiveTool* Tool)
diff --git a/Engine/Plugins/Editor/ModelingToolsEditorMode/Source/ModelingToolsEditorMode/Public/ModelingToolsEditorModeToolkit.h b/Engine/Plugins/Editor/ModelingToolsEditorMode/Source/ModelingToolsEditorMode/Public/ModelingToolsEditorModeToolkit.h
--- a/Engine/Plugins/Editor/ModelingToolsEditorMode/Source/ModelingToolsEditorMode/Public/ModelingToolsEditorModeToolkit.h
+++ b/Engine/Plugins/Editor/ModelingToolsEditorMode/Source/ModelingToolsEditorMode/Public/ModelingToolsEditorModeToolkit.h
@@ -8,13 +8,38 @@
 #include "Widgets/Input/STextComboBox.h"
 #include "Widgets/Input/SEditableTextBox.h"
 #include "StatusBarSubsystem.h"
+#include "UObject/GCObject.h"
+
+#include "ModelingToolsEditorModeToolkit.generated.h"
+
+#define LOCTEXT_NAMESPACE "ModelingToolsEditorModeToolkit"
 
 class STransformGizmoNumericalUIOverlay;
 class IDetailsView;
 class SButton;
 class STextBlock;
+class UInteractiveToolsPresetCollectionAsset;
+
+UCLASS()
+class UPresetSettingsProperties : public UObject
+{
+	GENERATED_BODY()
+public:
+
+	UPROPERTY(EditAnywhere, Category = "Presets", meta=(MetaClass = "InteractiveToolsPresetCollectionAsset"))
+	TArray< FSoftObjectPath > ActivePresetCollectionsPaths;
+};
+
+struct FToolPresetOption
+{
+	FString PresetLabel;
+	FString PresetTooltip;
+	FSlateIcon PresetIcon;
+	FString PresetName;
+	FSoftObjectPath PresetCollection;
+};
 
-class FModelingToolsEditorModeToolkit : public FModeToolkit
+class FModelingToolsEditorModeToolkit : public FModeToolkit, public FGCObject
 {
 public:
 
@@ -71,6 +96,11 @@ public:
 	// This is exposed only for the convenience of being able to create the numerical UI submenu
 	// in a non-member function in ModelingModeToolkit_Toolbars.cpp
 	TSharedPtr<STransformGizmoNumericalUIOverlay> GetGizmoNumericalUIOverlayWidget() { return GizmoNumericalUIOverlayWidget; }
+
+	/** GCObject interface */
+	virtual void AddReferencedObjects(FReferenceCollector& Collector) override;
+	virtual FString GetReferencerName() const override;
+
 private:
 	const static TArray<FName> PaletteNames_Standard;
 
@@ -118,27 +148,30 @@ private:
 	void SelectNewAssetPath() const;
 
 	// Presets
+	TObjectPtr<UPresetSettingsProperties> PresetSettings;
+	TSharedPtr<SWidget> MakePresetPanel();
 	FSoftObjectPath CurrentPreset;
-	TSharedPtr<FString> ActiveNamedPreset;
-	TArray<TSharedPtr<FString>> AvailablePresetsForTool;
+	FString NewPresetLabel;
+	FString NewPresetTooltip;
+	FSlateIcon NewPresetIcon;
+
+	TArray<TSharedPtr<FToolPresetOption>> AvailablePresetsForTool;
 	TSharedPtr<SEditableComboBox<TSharedPtr<FString>>> PresetComboBox;
 
+	TSharedRef<SWidget> GetPresetSettingsButtonContent();
+	TSharedRef<SWidget> GetPresetCreateButtonContent();
+
+	void CreateNewPresetInCollection(const FString& PresetLabel, FSoftObjectPath CollectionPath, const FString& ToolTip, FSlateIcon Icon);
+	void LoadPresetFromCollection(const FString& PresetName, FSoftObjectPath CollectionPath);
+
 	FString GetCurrentPresetPath() { return CurrentPreset.GetAssetPathString(); }
 	void HandlePresetAssetChanged(const FAssetData& InAssetData);
 	bool HandleFilterPresetAsset(const FAssetData& InAssetData);
-	void LoadActivePreset();
 	void SaveActivePreset();
 	
+	void RebuildPresetListForTool(bool bSettingsOpened);
 	TSharedRef<SWidget> MakePresetComboWidget(TSharedPtr<FString> InItem);
 	bool IsPresetEnabled() const;
-	void OnPresetChanged(TSharedPtr<FString> NewSelection, ESelectInfo::Type SelectInfo);
-	void OnPresetRenamed(const FText&, ETextCommit::Type);
-	FReply OnAddPreset();
-	FReply OnRemovePreset();
-	FText GetPresetComboBoxContent() const;
-	TSharedPtr<FString> GetPresetString() const;
-	FString GetRawPresetString() const;
-	void BuildPresetComboList();
 	void ClearPresetComboList();
 
 
@@ -149,3 +182,5 @@ private:
 	bool bFirstInitializeAfterModeSetup = true;
 	bool bShowActiveSelectionActions = true;
 };
+
+#undef LOCTEXT_NAMESPACE
\ No newline at end of file
diff --git a/Engine/Plugins/Experimental/ToolPresets/Source/PresetAsset/Private/AssetDefinition_PresetAsset.h b/Engine/Plugins/Experimental/ToolPresets/Source/PresetAsset/Private/AssetDefinition_PresetAsset.h
--- a/Engine/Plugins/Experimental/ToolPresets/Source/PresetAsset/Private/AssetDefinition_PresetAsset.h
+++ b/Engine/Plugins/Experimental/ToolPresets/Source/PresetAsset/Private/AssetDefinition_PresetAsset.h
@@ -12,19 +12,19 @@ struct FAssetCategoryPath;
 struct FAssetOpenArgs;
 
 UCLASS()
-class UAssetDefinition_PresetAsset : public UAssetDefinitionDefault
+class UAssetDefinition_InteractiveToolsPresetCollectionAsset : public UAssetDefinitionDefault
 {
 	GENERATED_BODY()
 
 public:
 	// UAssetDefinition Begin
-	virtual FText GetAssetDisplayName() const override { return NSLOCTEXT("AssetTypeActions", "AssetTypeActions_PresetAsset", "PresetAsset"); }
+	virtual FText GetAssetDisplayName() const override { return NSLOCTEXT("AssetTypeActions", "AssetTypeActions_InteractiveToolsPresetCollectionAsset", "Interactive Tools Preset Collection"); }
 	virtual FLinearColor GetAssetColor() const override { return FLinearColor(FColor(175, 0, 128)); }
-	virtual TSoftClassPtr<UObject> GetAssetClass() const override { return UPresetAsset::StaticClass(); }
+	virtual TSoftClassPtr<UObject> GetAssetClass() const override { return UInteractiveToolsPresetCollectionAsset::StaticClass(); }
 	virtual TConstArrayView<FAssetCategoryPath> GetAssetCategories() const override
 	{
 		return TConstArrayView<FAssetCategoryPath>();
 	}
-	virtual FText GetObjectDisplayNameText(UObject* Object) const override { return FText::FromString(TEXT("UPresetAsset")); }
+	virtual FText GetObjectDisplayNameText(UObject* Object) const override { return FText::FromString(TEXT("UInteractiveToolsPresetCollectionAsset")); }
 	// UAssetDefinition End
 };
\ No newline at end of file
diff --git a/Engine/Plugins/Experimental/ToolPresets/Source/PresetAsset/Public/PresetAsset.h b/Engine/Plugins/Experimental/ToolPresets/Source/PresetAsset/Public/PresetAsset.h
--- a/Engine/Plugins/Experimental/ToolPresets/Source/PresetAsset/Public/PresetAsset.h
+++ b/Engine/Plugins/Experimental/ToolPresets/Source/PresetAsset/Public/PresetAsset.h
@@ -13,28 +13,38 @@
 #define LOCTEXT_NAMESPACE "PresetAsset"
 
 USTRUCT(BlueprintType)
-struct FPropertyStore
+struct FInteractiveToolPresetDefintion
 {
 	GENERATED_BODY()
 public:
 	UPROPERTY(VisibleAnywhere, Category = "PresetAsset")
 	TArray< TObjectPtr<UObject> > Properties;
+
+	UPROPERTY(EditAnywhere, Category = "PresetAsset")
+	FString Label;
+
+	UPROPERTY(EditAnywhere, Category = "PresetAsset")
+	FString Tooltip;
+
+	// TODO: Investigate how to store things like icons. This doesn't work as written, as FSlateBrush can't be serialized?
+	//UPROPERTY(EditAnywhere, Category = "PresetAsset")
+	//FSlateBrush Icon;
 };
 
 USTRUCT(BlueprintType)
-struct FNamedPropertyStore
+struct FInteractiveToolPresetStore
 {
 	GENERATED_BODY()
 public:
 	UPROPERTY(EditAnywhere, Category = "PresetAsset")
-	TMap< FString, FPropertyStore> Store;
+	TMap< FString, FInteractiveToolPresetDefintion> NamedPresets;
 };
 
 /**
  * Implements an asset that can be used to store tool settings as a named preset
  */
 UCLASS(BlueprintType, hidecategories=(Object))
-class PRESETASSET_API UPresetAsset
+class PRESETASSET_API UInteractiveToolsPresetCollectionAsset
 	: public UObject
 {
 	GENERATED_BODY()
@@ -49,21 +59,21 @@ public:
 	// support for adding, removing, renaming, saving and retrieving presets.
 
 	UPROPERTY(VisibleAnywhere, Category = "PresetAsset")
-	TMap<FString, FNamedPropertyStore > StoredProperties;
+	TMap<FString, FInteractiveToolPresetStore > PerToolPresets;
 
 };
 
 
 UCLASS(hidecategories = Object)
-class PRESETASSET_API UPresetAssetFactory : public UFactory
+class PRESETASSET_API UInteractiveToolsPresetCollectionAssetFactory : public UFactory
 {
 	GENERATED_BODY()
 
 public:
 
-	UPresetAssetFactory()
+	UInteractiveToolsPresetCollectionAssetFactory()
 	{
-		SupportedClass = UPresetAsset::StaticClass();
+		SupportedClass = UInteractiveToolsPresetCollectionAsset::StaticClass();
 
 		bCreateNew = true;
 		bEditAfterNew = true;
@@ -81,11 +91,11 @@ public:
 
 	UObject* FactoryCreateNew(UClass* Class, UObject* InParent, FName Name, EObjectFlags Flags, UObject* Context, FFeedbackContext* Warn) override
 	{
-		UPresetAsset* Preset = nullptr;
+		UInteractiveToolsPresetCollectionAsset* Preset = nullptr;
 		if (ensure(SupportedClass == Class))
 		{
 			ensure(0 != (RF_Public & Flags));
-			Preset = NewObject<UPresetAsset>(InParent, Class, Name, Flags);
+			Preset = NewObject<UInteractiveToolsPresetCollectionAsset>(InParent, Class, Name, Flags);
 		}
 		return Preset;
 	}
