[Engine/Plugins]

diff --git a/Engine/Plugins/Editor/DataValidation/Source/DataValidation/Public/EditorValidatorBase.h b/Engine/Plugins/Editor/DataValidation/Source/DataValidation/Public/EditorValidatorBase.h
--- a/Engine/Plugins/Editor/DataValidation/Source/DataValidation/Public/EditorValidatorBase.h
+++ b/Engine/Plugins/Editor/DataValidation/Source/DataValidation/Public/EditorValidatorBase.h
@@ -2,20 +2,13 @@
 
 #pragma once
 
-#include "EditorSubsystem.h"
 
-#include "UObject/ObjectMacros.h"
-#include "UObject/UObjectGlobals.h"
-#include "Templates/SubclassOf.h"
-#include "UObject/Interface.h"
-#include "UObject/ObjectMacros.h"
-#include "Logging/LogVerbosity.h"
-#include "Internationalization/Text.h"
-#include "UObject/TextProperty.h"
-#include "DataValidationModule.h"
 
+#include "UObject/Package.h"
 #include "EditorValidatorBase.generated.h"
 
+enum class EDataValidationUsecase : uint8;
+
 /*
 * The EditorValidatorBase is a class which verifies that an asset meets a specific ruleset.
 * It should be used when checking engine-level classes, as UObject::IsDataValid requires
@@ -91,3 +84,11 @@ private:
 
 
 
+
+#if UE_ENABLE_INCLUDE_ORDER_DEPRECATED_IN_5_2
+#include "DataValidationModule.h"
+#include "EditorSubsystem.h"
+#include "Templates/SubclassOf.h"
+#include "UObject/Interface.h"
+#include "UObject/TextProperty.h"
+#endif
diff --git a/Engine/Plugins/Editor/DataValidation/Source/DataValidation/Public/EditorValidatorSubsystem.h b/Engine/Plugins/Editor/DataValidation/Source/DataValidation/Public/EditorValidatorSubsystem.h
--- a/Engine/Plugins/Editor/DataValidation/Source/DataValidation/Public/EditorValidatorSubsystem.h
+++ b/Engine/Plugins/Editor/DataValidation/Source/DataValidation/Public/EditorValidatorSubsystem.h
@@ -4,19 +4,14 @@
 
 #include "EditorSubsystem.h"
 
-#include "UObject/ObjectMacros.h"
-#include "UObject/UObjectGlobals.h"
-#include "Templates/SubclassOf.h"
-#include "UObject/Interface.h"
-#include "UObject/ObjectMacros.h"
-#include "EditorValidatorBase.h"
-#include "Engine/EngineTypes.h"
-#include "AssetRegistry/AssetData.h"
-#include "Logging/LogMacros.h"
 #include "DataValidationModule.h"
 
+#include "UObject/Package.h"
 #include "EditorValidatorSubsystem.generated.h"
 
+class UEditorValidatorBase;
+struct FDirectoryPath;
+
 struct FAssetData;
 class ISourceControlChangelist;
 typedef TSharedPtr<class ISourceControlChangelist, ESPMode::ThreadSafe> FSourceControlChangelistPtr;
@@ -218,3 +213,11 @@ protected:
 
 };
 
+
+#if UE_ENABLE_INCLUDE_ORDER_DEPRECATED_IN_5_2
+#include "AssetRegistry/AssetData.h"
+#include "EditorValidatorBase.h"
+#include "Engine/EngineTypes.h"
+#include "Templates/SubclassOf.h"
+#include "UObject/Interface.h"
+#endif
diff --git a/Engine/Plugins/Editor/EngineAssetDefinitions/Source/Private/Animation/AssetDefinition_AimOffset.h b/Engine/Plugins/Editor/EngineAssetDefinitions/Source/Private/Animation/AssetDefinition_AimOffset.h
--- a/Engine/Plugins/Editor/EngineAssetDefinitions/Source/Private/Animation/AssetDefinition_AimOffset.h
+++ b/Engine/Plugins/Editor/EngineAssetDefinitions/Source/Private/Animation/AssetDefinition_AimOffset.h
@@ -3,10 +3,12 @@
 #pragma once
 
 #include "Animation/AimOffsetBlendSpace.h"
-#include "AssetDefinition_BlendSpace.h"
+#include "Animation/AssetDefinition_BlendSpace.h"
 
 #include "AssetDefinition_AimOffset.generated.h"
 
+class UAssetDefinition_BlendSpace;
+
 UCLASS()
 class UAssetDefinition_AimOffset : public UAssetDefinition_BlendSpace
 {
@@ -18,4 +20,4 @@ public:
 	virtual FLinearColor GetAssetColor() const override { return FLinearColor(FColor(0, 162, 232)); }
 	virtual TSoftClassPtr<UObject> GetAssetClass() const override { return UAimOffsetBlendSpace::StaticClass(); }
 	// UAssetDefinition End
-};
\ No newline at end of file
+};
diff --git a/Engine/Plugins/Editor/EngineAssetDefinitions/Source/Private/Animation/AssetDefinition_AimOffset1D.h b/Engine/Plugins/Editor/EngineAssetDefinitions/Source/Private/Animation/AssetDefinition_AimOffset1D.h
--- a/Engine/Plugins/Editor/EngineAssetDefinitions/Source/Private/Animation/AssetDefinition_AimOffset1D.h
+++ b/Engine/Plugins/Editor/EngineAssetDefinitions/Source/Private/Animation/AssetDefinition_AimOffset1D.h
@@ -3,10 +3,12 @@
 #pragma once
 
 #include "Animation/AimOffsetBlendSpace1D.h"
-#include "AssetDefinition_BlendSpace1D.h"
+#include "Animation/AssetDefinition_BlendSpace1D.h"
 
 #include "AssetDefinition_AimOffset1D.generated.h"
 
+class UAssetDefinition_BlendSpace1D;
+
 UCLASS()
 class UAssetDefinition_AimOffset1D : public UAssetDefinition_BlendSpace1D
 {
@@ -23,4 +25,4 @@ public:
 		return Categories;
 	}
 	// UAssetDefinition End
-};
\ No newline at end of file
+};
diff --git a/Engine/Plugins/Editor/EngineAssetDefinitions/Source/Private/Animation/AssetDefinition_AnimBlueprint.cpp b/Engine/Plugins/Editor/EngineAssetDefinitions/Source/Private/Animation/AssetDefinition_AnimBlueprint.cpp
--- a/Engine/Plugins/Editor/EngineAssetDefinitions/Source/Private/Animation/AssetDefinition_AnimBlueprint.cpp
+++ b/Engine/Plugins/Editor/EngineAssetDefinitions/Source/Private/Animation/AssetDefinition_AnimBlueprint.cpp
@@ -2,28 +2,24 @@
 
 #include "Animation/AssetDefinition_AnimBlueprint.h"
 
-#include "AnimationEditorUtils.h"
 #include "ContentBrowserMenuContexts.h"
+#include "Misc/PackageName.h"
 #include "ToolMenus.h"
-#include "Widgets/Layout/SBorder.h"
 #include "Misc/MessageDialog.h"
-#include "Widgets/Images/SImage.h"
-#include "Styling/AppStyle.h"
-#include "Animation/AnimInstance.h"
 #include "Factories/AnimBlueprintFactory.h"
 #include "ContentBrowserModule.h"
 #include "Framework/Application/SlateApplication.h"
 #include "SBlueprintDiff.h"
 #include "SSkeletonWidget.h"
 #include "Styling/SlateIconFinder.h"
-#include "IContentBrowserSingleton.h"
 #include "Kismet2/KismetEditorUtilities.h"
-#include "Subsystems/AssetEditorSubsystem.h"
-#include "Editor.h"
 #include "IAnimationBlueprintEditorModule.h"
 #include "IAssetTools.h"
 #include "Algo/AllOf.h"
 #include "Algo/NoneOf.h"
+#include "ToolMenu.h"
+#include "ToolMenuEntry.h"
+#include "ToolMenuSection.h"
 
 #define LOCTEXT_NAMESPACE "AssetTypeActions"
 
diff --git a/Engine/Plugins/Editor/EngineAssetDefinitions/Source/Private/Animation/AssetDefinition_AnimComposite.h b/Engine/Plugins/Editor/EngineAssetDefinitions/Source/Private/Animation/AssetDefinition_AnimComposite.h
--- a/Engine/Plugins/Editor/EngineAssetDefinitions/Source/Private/Animation/AssetDefinition_AnimComposite.h
+++ b/Engine/Plugins/Editor/EngineAssetDefinitions/Source/Private/Animation/AssetDefinition_AnimComposite.h
@@ -3,10 +3,12 @@
 #pragma once
 
 #include "Animation/AnimComposite.h"
-#include "AssetDefinition_AnimationAsset.h"
+#include "Animation/AssetDefinition_AnimationAsset.h"
 
 #include "AssetDefinition_AnimComposite.generated.h"
 
+class UAssetDefinition_AnimationAsset;
+
 UCLASS()
 class UAssetDefinition_AnimComposite : public UAssetDefinition_AnimationAsset
 {
diff --git a/Engine/Plugins/Editor/EngineAssetDefinitions/Source/Private/Animation/AssetDefinition_AnimMontage.cpp b/Engine/Plugins/Editor/EngineAssetDefinitions/Source/Private/Animation/AssetDefinition_AnimMontage.cpp
--- a/Engine/Plugins/Editor/EngineAssetDefinitions/Source/Private/Animation/AssetDefinition_AnimMontage.cpp
+++ b/Engine/Plugins/Editor/EngineAssetDefinitions/Source/Private/Animation/AssetDefinition_AnimMontage.cpp
@@ -5,6 +5,7 @@
 #include "AnimationEditorUtils.h"
 #include "ContentBrowserMenuContexts.h"
 #include "ToolMenu.h"
+#include "ToolMenuSection.h"
 #include "ToolMenus.h"
 
 #define LOCTEXT_NAMESPACE "AssetTypeActions"
diff --git a/Engine/Plugins/Editor/EngineAssetDefinitions/Source/Private/Animation/AssetDefinition_AnimMontage.h b/Engine/Plugins/Editor/EngineAssetDefinitions/Source/Private/Animation/AssetDefinition_AnimMontage.h
--- a/Engine/Plugins/Editor/EngineAssetDefinitions/Source/Private/Animation/AssetDefinition_AnimMontage.h
+++ b/Engine/Plugins/Editor/EngineAssetDefinitions/Source/Private/Animation/AssetDefinition_AnimMontage.h
@@ -3,10 +3,12 @@
 #pragma once
 
 #include "Animation/AnimMontage.h"
-#include "AssetDefinition_AnimationAsset.h"
+#include "Animation/AssetDefinition_AnimationAsset.h"
 
 #include "AssetDefinition_AnimMontage.generated.h"
 
+class UAssetDefinition_AnimationAsset;
+
 UCLASS()
 class UAssetDefinition_AnimMontage : public UAssetDefinition_AnimationAsset
 {
diff --git a/Engine/Plugins/Editor/EngineAssetDefinitions/Source/Private/Animation/AssetDefinition_AnimSequence.cpp b/Engine/Plugins/Editor/EngineAssetDefinitions/Source/Private/Animation/AssetDefinition_AnimSequence.cpp
--- a/Engine/Plugins/Editor/EngineAssetDefinitions/Source/Private/Animation/AssetDefinition_AnimSequence.cpp
+++ b/Engine/Plugins/Editor/EngineAssetDefinitions/Source/Private/Animation/AssetDefinition_AnimSequence.cpp
@@ -1,9 +1,8 @@
 // Copyright Epic Games, Inc. All Rights Reserved.
 
 #include "AssetDefinition_AnimSequence.h"
-#include "Animation/AnimSequence.h"
+#include "ToolMenuSection.h"
 #include "ToolMenus.h"
-#include "Styling/AppStyle.h"
 #include "EditorReimportHandler.h"
 #include "Animation/AnimMontage.h"
 #include "Factories/AnimCompositeFactory.h"
@@ -14,7 +13,6 @@
 #include "Animation/AnimStreamable.h"
 #include "Animation/PoseAsset.h"
 #include "IAssetTools.h"
-#include "Algo/Transform.h"
 #include "ToolMenu.h"
 #include "ToolMenus.h"
 #include "ContentBrowserMenuContexts.h"
diff --git a/Engine/Plugins/Editor/EngineAssetDefinitions/Source/Private/Animation/AssetDefinition_AnimationAsset.cpp b/Engine/Plugins/Editor/EngineAssetDefinitions/Source/Private/Animation/AssetDefinition_AnimationAsset.cpp
--- a/Engine/Plugins/Editor/EngineAssetDefinitions/Source/Private/Animation/AssetDefinition_AnimationAsset.cpp
+++ b/Engine/Plugins/Editor/EngineAssetDefinitions/Source/Private/Animation/AssetDefinition_AnimationAsset.cpp
@@ -1,21 +1,19 @@
 // Copyright Epic Games, Inc. All Rights Reserved.
 
 #include "AssetDefinition_AnimationAsset.h"
-#include "Framework/MultiBox/MultiBoxBuilder.h"
+#include "IAnimationEditor.h"
 #include "Misc/MessageDialog.h"
-#include "Styling/AppStyle.h"
 #include "ThumbnailRendering/SceneThumbnailInfo.h"
 #include "IAssetTools.h"
 #include "ContentBrowserMenuContexts.h"
 #include "SSkeletonWidget.h"
 #include "IAnimationEditorModule.h"
 #include "Preferences/PersonaOptions.h"
+#include "ToolMenu.h"
 #include "ToolMenus.h"
 #include "Framework/Application/SlateApplication.h"
-#include "Algo/Transform.h"
 #include "IPersonaToolkit.h"
-#include "Subsystems/AssetEditorSubsystem.h"
-#include "Editor.h"
+#include "ToolMenuSection.h"
 
 #define LOCTEXT_NAMESPACE "AssetTypeActions"
 
diff --git a/Engine/Plugins/Editor/EngineAssetDefinitions/Source/Private/Animation/AssetDefinition_BlendSpace.h b/Engine/Plugins/Editor/EngineAssetDefinitions/Source/Private/Animation/AssetDefinition_BlendSpace.h
--- a/Engine/Plugins/Editor/EngineAssetDefinitions/Source/Private/Animation/AssetDefinition_BlendSpace.h
+++ b/Engine/Plugins/Editor/EngineAssetDefinitions/Source/Private/Animation/AssetDefinition_BlendSpace.h
@@ -2,11 +2,13 @@
 
 #pragma once
 
+#include "Animation/AssetDefinition_AnimationAsset.h"
 #include "Animation/BlendSpace.h"
-#include "AssetDefinition_AnimationAsset.h"
 
 #include "AssetDefinition_BlendSpace.generated.h"
 
+class UAssetDefinition_AnimationAsset;
+
 UCLASS()
 class UAssetDefinition_BlendSpace : public UAssetDefinition_AnimationAsset
 {
@@ -18,4 +20,4 @@ public:
 	virtual FLinearColor GetAssetColor() const override { return FLinearColor(FColor(255, 168, 111)); }
 	virtual TSoftClassPtr<UObject> GetAssetClass() const override { return UBlendSpace::StaticClass(); }
 	// UAssetDefinition End
-};
\ No newline at end of file
+};
diff --git a/Engine/Plugins/Editor/EngineAssetDefinitions/Source/Private/Animation/AssetDefinition_BlendSpace1D.h b/Engine/Plugins/Editor/EngineAssetDefinitions/Source/Private/Animation/AssetDefinition_BlendSpace1D.h
--- a/Engine/Plugins/Editor/EngineAssetDefinitions/Source/Private/Animation/AssetDefinition_BlendSpace1D.h
+++ b/Engine/Plugins/Editor/EngineAssetDefinitions/Source/Private/Animation/AssetDefinition_BlendSpace1D.h
@@ -2,11 +2,13 @@
 
 #pragma once
 
+#include "Animation/AssetDefinition_AnimationAsset.h"
 #include "Animation/BlendSpace1D.h"
-#include "AssetDefinition_AnimationAsset.h"
 
 #include "AssetDefinition_BlendSpace1D.generated.h"
 
+class UAssetDefinition_AnimationAsset;
+
 UCLASS()
 class UAssetDefinition_BlendSpace1D : public UAssetDefinition_AnimationAsset
 {
@@ -23,4 +25,4 @@ public:
 		return Categories;
 	}
 	// UAssetDefinition End
-};
\ No newline at end of file
+};
diff --git a/Engine/Plugins/Editor/EngineAssetDefinitions/Source/Private/Animation/AssetDefinition_PoseAsset.h b/Engine/Plugins/Editor/EngineAssetDefinitions/Source/Private/Animation/AssetDefinition_PoseAsset.h
--- a/Engine/Plugins/Editor/EngineAssetDefinitions/Source/Private/Animation/AssetDefinition_PoseAsset.h
+++ b/Engine/Plugins/Editor/EngineAssetDefinitions/Source/Private/Animation/AssetDefinition_PoseAsset.h
@@ -2,11 +2,13 @@
 
 #pragma once
 
+#include "Animation/AssetDefinition_AnimationAsset.h"
 #include "Animation/PoseAsset.h"
-#include "AssetDefinition_AnimationAsset.h"
 
 #include "AssetDefinition_PoseAsset.generated.h"
 
+class UAssetDefinition_AnimationAsset;
+
 UCLASS()
 class UAssetDefinition_PoseAsset : public UAssetDefinition_AnimationAsset
 {
diff --git a/Engine/Plugins/Editor/EngineAssetDefinitions/Source/Private/Script/AssetDefinition_Blueprint.cpp b/Engine/Plugins/Editor/EngineAssetDefinitions/Source/Private/Script/AssetDefinition_Blueprint.cpp
--- a/Engine/Plugins/Editor/EngineAssetDefinitions/Source/Private/Script/AssetDefinition_Blueprint.cpp
+++ b/Engine/Plugins/Editor/EngineAssetDefinitions/Source/Private/Script/AssetDefinition_Blueprint.cpp
@@ -5,7 +5,6 @@
 #include "AssetDefinitionRegistry.h"
 #include "AssetToolsModule.h"
 #include "Blueprint/BlueprintSupport.h"
-#include "BlueprintEditorModule.h"
 #include "ContentBrowserMenuContexts.h"
 #include "ContentBrowserModule.h"
 #include "IAssetTools.h"
@@ -19,12 +18,10 @@
 #include "Misc/PackageName.h"
 #include "ToolMenu.h"
 #include "ToolMenuSection.h"
-#include "Toolkits/AssetEditorToolkit.h"
 #include "BlueprintEditor.h"
 #include "SBlueprintDiff.h"
 #include "Kismet2/BlueprintEditorUtils.h"
 #include "ThumbnailRendering/SceneThumbnailInfo.h"
-#include "Misc/AssetFilterData.h"
 
 #define LOCTEXT_NAMESPACE "UAssetDefinition_Blueprint"
 
diff --git a/Engine/Plugins/Editor/EngineAssetDefinitions/Source/Private/Script/AssetDefinition_BlueprintGeneratedClass.cpp b/Engine/Plugins/Editor/EngineAssetDefinitions/Source/Private/Script/AssetDefinition_BlueprintGeneratedClass.cpp
--- a/Engine/Plugins/Editor/EngineAssetDefinitions/Source/Private/Script/AssetDefinition_BlueprintGeneratedClass.cpp
+++ b/Engine/Plugins/Editor/EngineAssetDefinitions/Source/Private/Script/AssetDefinition_BlueprintGeneratedClass.cpp
@@ -4,15 +4,17 @@
 
 #include "AssetDefinitionRegistry.h"
 #include "ContentBrowserMenuContexts.h"
-#include "Engine/BlueprintGeneratedClass.h"
-#include "ToolMenus.h"
-#include "Misc/PackageName.h"
-#include "Misc/MessageDialog.h"
+#include "Blueprint/BlueprintSupport.h"
+#include "ContentBrowserModule.h"
 #include "Factories/BlueprintFactory.h"
-#include "Kismet2/KismetEditorUtilities.h"
 #include "IContentBrowserSingleton.h"
-#include "ContentBrowserModule.h"
+#include "Kismet2/KismetEditorUtilities.h"
+#include "Misc/MessageDialog.h"
+#include "Misc/PackageName.h"
 #include "IAssetTools.h"
+#include "ToolMenu.h"
+#include "ToolMenus.h"
+#include "ToolMenuSection.h"
 
 #define LOCTEXT_NAMESPACE "AssetTypeActions"
 
diff --git a/Engine/Plugins/Editor/EngineAssetDefinitions/Source/Private/Script/AssetDefinition_Class.cpp b/Engine/Plugins/Editor/EngineAssetDefinitions/Source/Private/Script/AssetDefinition_Class.cpp
--- a/Engine/Plugins/Editor/EngineAssetDefinitions/Source/Private/Script/AssetDefinition_Class.cpp
+++ b/Engine/Plugins/Editor/EngineAssetDefinitions/Source/Private/Script/AssetDefinition_Class.cpp
@@ -1,18 +1,20 @@
 // Copyright Epic Games, Inc. All Rights Reserved.
 
 #include "AssetDefinition_Class.h"
-#include "ToolMenus.h"
-#include "AssetRegistry/AssetData.h"
+#include "Framework/Docking/TabManager.h"
 #include "HAL/FileManager.h"
-#include "Styling/AppStyle.h"
+#include "IAssetTools.h"
+#include "Misc/Paths.h"
 #include "SourceCodeNavigation.h"
-#include "AddToProjectConfig.h"
-#include "AssetToolsModule.h"
 #include "ContentBrowserMenuContexts.h"
 #include "GameProjectGenerationModule.h"
 #include "Kismet2/KismetEditorUtilities.h"
 #include "Editor/UnrealEdEngine.h"
+#include "ModuleDescriptor.h"
 #include "Preferences/UnrealEdOptions.h"
+#include "ToolMenu.h"
+#include "ToolMenus.h"
+#include "ToolMenuSection.h"
 #include "UnrealEdGlobals.h"
 
 #define LOCTEXT_NAMESPACE "AssetTypeActions"
diff --git a/Engine/Plugins/Editor/EngineAssetDefinitions/Source/Public/Script/AssetDefinition_Blueprint.h b/Engine/Plugins/Editor/EngineAssetDefinitions/Source/Public/Script/AssetDefinition_Blueprint.h
--- a/Engine/Plugins/Editor/EngineAssetDefinitions/Source/Public/Script/AssetDefinition_Blueprint.h
+++ b/Engine/Plugins/Editor/EngineAssetDefinitions/Source/Public/Script/AssetDefinition_Blueprint.h
@@ -2,11 +2,12 @@
 
 #pragma once
 
-#include "AssetDefinition_ClassTypeBase.h"
 #include "Engine/Blueprint.h"
-
+#include "Script/AssetDefinition_ClassTypeBase.h"
 #include "AssetDefinition_Blueprint.generated.h"
 
+struct FAssetFilterData;
+
 struct FToolMenuContext;
 struct FAssetData;
 class UFactory;
diff --git a/Engine/Plugins/Editor/GLTFImporter/Source/GLTFCore/Public/GLTFAccessor.h b/Engine/Plugins/Editor/GLTFImporter/Source/GLTFCore/Public/GLTFAccessor.h
--- a/Engine/Plugins/Editor/GLTFImporter/Source/GLTFCore/Public/GLTFAccessor.h
+++ b/Engine/Plugins/Editor/GLTFImporter/Source/GLTFCore/Public/GLTFAccessor.h
@@ -2,7 +2,7 @@
 
 #pragma once
 
-#include "CoreMinimal.h"
+#include "Math/Matrix.h"
 
 struct FMD5Hash;
 
@@ -229,3 +229,7 @@ namespace GLTF
 	}
 
 }  // namespace GLTF
+
+#if UE_ENABLE_INCLUDE_ORDER_DEPRECATED_IN_5_2
+#include "CoreMinimal.h"
+#endif
diff --git a/Engine/Plugins/Editor/GLTFImporter/Source/GLTFCore/Public/GLTFAnimation.h b/Engine/Plugins/Editor/GLTFImporter/Source/GLTFCore/Public/GLTFAnimation.h
--- a/Engine/Plugins/Editor/GLTFImporter/Source/GLTFCore/Public/GLTFAnimation.h
+++ b/Engine/Plugins/Editor/GLTFImporter/Source/GLTFCore/Public/GLTFAnimation.h
@@ -2,7 +2,9 @@
 
 #pragma once
 
-#include "CoreMinimal.h"
+#include "Containers/UnrealString.h"
+
+struct FScriptContainerElement;
 
 namespace GLTF
 {
@@ -74,3 +76,7 @@ namespace GLTF
 	};
 
 }  // namespace GLTF
+
+#if UE_ENABLE_INCLUDE_ORDER_DEPRECATED_IN_5_2
+#include "CoreMinimal.h"
+#endif
diff --git a/Engine/Plugins/Editor/GLTFImporter/Source/GLTFCore/Public/GLTFAsset.h b/Engine/Plugins/Editor/GLTFImporter/Source/GLTFCore/Public/GLTFAsset.h
--- a/Engine/Plugins/Editor/GLTFImporter/Source/GLTFCore/Public/GLTFAsset.h
+++ b/Engine/Plugins/Editor/GLTFImporter/Source/GLTFCore/Public/GLTFAsset.h
@@ -2,14 +2,23 @@
 
 #pragma once
 
-#include "GLTFAnimation.h"
-#include "GLTFMaterial.h"
-#include "GLTFMesh.h"
-#include "GLTFNode.h"
-#include "GLTFTexture.h"
 
 #include "Containers/Set.h"
-#include "CoreMinimal.h"
+
+namespace GLTF { struct FAnimation; }
+namespace GLTF { struct FBuffer; }
+namespace GLTF { struct FBufferView; }
+namespace GLTF { struct FCamera; }
+namespace GLTF { struct FImage; }
+namespace GLTF { struct FLight; }
+namespace GLTF { struct FMaterial; }
+namespace GLTF { struct FMesh; }
+namespace GLTF { struct FNode; }
+namespace GLTF { struct FSampler; }
+namespace GLTF { struct FSkinInfo; }
+namespace GLTF { struct FTexture; }
+namespace GLTF { struct FValidAccessor; }
+struct FScriptContainerElement;
 
 namespace GLTF
 {
@@ -131,3 +140,12 @@ namespace GLTF
 	GLTFCORE_API const TCHAR* ToString(EExtension Extension);
 
 }  // namespace GLTF
+
+#if UE_ENABLE_INCLUDE_ORDER_DEPRECATED_IN_5_2
+#include "CoreMinimal.h"
+#include "GLTFAnimation.h"
+#include "GLTFMaterial.h"
+#include "GLTFMesh.h"
+#include "GLTFNode.h"
+#include "GLTFTexture.h"
+#endif
diff --git a/Engine/Plugins/Editor/GLTFImporter/Source/GLTFCore/Public/GLTFCoreModule.h b/Engine/Plugins/Editor/GLTFImporter/Source/GLTFCore/Public/GLTFCoreModule.h
--- a/Engine/Plugins/Editor/GLTFImporter/Source/GLTFCore/Public/GLTFCoreModule.h
+++ b/Engine/Plugins/Editor/GLTFImporter/Source/GLTFCore/Public/GLTFCoreModule.h
@@ -2,8 +2,6 @@
 
 #pragma once
 
-#include "CoreMinimal.h"
-#include "Modules/ModuleInterface.h"
 #include "Modules/ModuleManager.h"
 
 /**
@@ -32,3 +30,7 @@ public:
 		return FModuleManager::Get().IsModuleLoaded("GLTFCore");
 	}
 };
+
+#if UE_ENABLE_INCLUDE_ORDER_DEPRECATED_IN_5_2
+#include "CoreMinimal.h"
+#endif
diff --git a/Engine/Plugins/Editor/GLTFImporter/Source/GLTFCore/Public/GLTFLogger.h b/Engine/Plugins/Editor/GLTFImporter/Source/GLTFCore/Public/GLTFLogger.h
--- a/Engine/Plugins/Editor/GLTFImporter/Source/GLTFCore/Public/GLTFLogger.h
+++ b/Engine/Plugins/Editor/GLTFImporter/Source/GLTFCore/Public/GLTFLogger.h
@@ -2,8 +2,6 @@
 
 #pragma once
 
-#include "Containers/Array.h"
-#include "Containers/UnrealString.h"
 #include "Templates/Tuple.h"
 
 namespace GLTF
@@ -29,3 +27,7 @@ namespace GLTF
 		return Messages;
 	}
 }  // namespace GLTF
+
+#if UE_ENABLE_INCLUDE_ORDER_DEPRECATED_IN_5_2
+#include "Containers/UnrealString.h"
+#endif
diff --git a/Engine/Plugins/Editor/GLTFImporter/Source/GLTFCore/Public/GLTFMaterial.h b/Engine/Plugins/Editor/GLTFImporter/Source/GLTFCore/Public/GLTFMaterial.h
--- a/Engine/Plugins/Editor/GLTFImporter/Source/GLTFCore/Public/GLTFMaterial.h
+++ b/Engine/Plugins/Editor/GLTFImporter/Source/GLTFCore/Public/GLTFMaterial.h
@@ -2,7 +2,7 @@
 
 #pragma once
 
-#include "CoreMinimal.h"
+#include "Math/Vector4.h"
 #include "Misc/SecureHash.h"
 
 namespace GLTF
@@ -195,3 +195,7 @@ namespace GLTF
 	};
 }  // namespace GLTF
 
+
+#if UE_ENABLE_INCLUDE_ORDER_DEPRECATED_IN_5_2
+#include "CoreMinimal.h"
+#endif
diff --git a/Engine/Plugins/Editor/GLTFImporter/Source/GLTFCore/Public/GLTFMaterialExpressions.h b/Engine/Plugins/Editor/GLTFImporter/Source/GLTFCore/Public/GLTFMaterialExpressions.h
--- a/Engine/Plugins/Editor/GLTFImporter/Source/GLTFCore/Public/GLTFMaterialExpressions.h
+++ b/Engine/Plugins/Editor/GLTFImporter/Source/GLTFCore/Public/GLTFMaterialExpressions.h
@@ -2,7 +2,6 @@
 
 #pragma once
 
-#include "UObject/StrongObjectPtr.h"
 #include "Misc/SecureHash.h"
 
 namespace GLTF
@@ -472,3 +471,7 @@ namespace GLTF
 	}
 
 }  // namespace GLTF
+
+#if UE_ENABLE_INCLUDE_ORDER_DEPRECATED_IN_5_2
+#include "UObject/StrongObjectPtr.h"
+#endif
diff --git a/Engine/Plugins/Editor/GLTFImporter/Source/GLTFCore/Public/GLTFMeshFactory.h b/Engine/Plugins/Editor/GLTFImporter/Source/GLTFCore/Public/GLTFMeshFactory.h
--- a/Engine/Plugins/Editor/GLTFImporter/Source/GLTFCore/Public/GLTFMeshFactory.h
+++ b/Engine/Plugins/Editor/GLTFImporter/Source/GLTFCore/Public/GLTFMeshFactory.h
@@ -2,9 +2,10 @@
 
 #pragma once
 
-#include "CoreMinimal.h"
+#include "Containers/Map.h"
 #include "GLTFLogger.h"
-#include "MeshTypes.h"
+
+struct FVertexID;
 
 struct FMeshDescription;
 
@@ -37,4 +38,9 @@ namespace GLTF
 	private:
 		TUniquePtr<FMeshFactoryImpl> Impl;
 	};
-}
\ No newline at end of file
+}
+
+#if UE_ENABLE_INCLUDE_ORDER_DEPRECATED_IN_5_2
+#include "CoreMinimal.h"
+#include "MeshTypes.h"
+#endif
diff --git a/Engine/Plugins/Editor/GLTFImporter/Source/GLTFCore/Public/GLTFNode.h b/Engine/Plugins/Editor/GLTFImporter/Source/GLTFCore/Public/GLTFNode.h
--- a/Engine/Plugins/Editor/GLTFImporter/Source/GLTFCore/Public/GLTFNode.h
+++ b/Engine/Plugins/Editor/GLTFImporter/Source/GLTFCore/Public/GLTFNode.h
@@ -2,12 +2,10 @@
 
 #pragma once
 
-#include "GLTFMaterial.h"
-#include "GLTFMesh.h"
-#include "GLTFTexture.h"
 
-#include "Containers/Set.h"
-#include "CoreMinimal.h"
+#include "Math/Transform.h"
+
+namespace GLTF { struct FAccessor; }
 
 namespace GLTF
 {
@@ -150,3 +148,10 @@ namespace GLTF
 	};
 
 }  // namespace GLTF
+
+#if UE_ENABLE_INCLUDE_ORDER_DEPRECATED_IN_5_2
+#include "CoreMinimal.h"
+#include "GLTFMaterial.h"
+#include "GLTFMesh.h"
+#include "GLTFTexture.h"
+#endif
diff --git a/Engine/Plugins/Editor/GLTFImporter/Source/GLTFCore/Public/GLTFReader.h b/Engine/Plugins/Editor/GLTFImporter/Source/GLTFCore/Public/GLTFReader.h
--- a/Engine/Plugins/Editor/GLTFImporter/Source/GLTFCore/Public/GLTFReader.h
+++ b/Engine/Plugins/Editor/GLTFImporter/Source/GLTFCore/Public/GLTFReader.h
@@ -4,7 +4,7 @@
 
 #include "GLTFLogger.h"
 
-#include "CoreMinimal.h"
+#include "Templates/SharedPointer.h"
 
 class FJsonObject;
 
@@ -83,3 +83,7 @@ namespace GLTF
 	};
 
 }  // namespace GLTF
+
+#if UE_ENABLE_INCLUDE_ORDER_DEPRECATED_IN_5_2
+#include "CoreMinimal.h"
+#endif
diff --git a/Engine/Plugins/Editor/GLTFImporter/Source/GLTFCore/Public/GLTFTexture.h b/Engine/Plugins/Editor/GLTFImporter/Source/GLTFCore/Public/GLTFTexture.h
--- a/Engine/Plugins/Editor/GLTFImporter/Source/GLTFCore/Public/GLTFTexture.h
+++ b/Engine/Plugins/Editor/GLTFImporter/Source/GLTFCore/Public/GLTFTexture.h
@@ -2,7 +2,9 @@
 
 #pragma once
 
-#include "CoreMinimal.h"
+#include "Containers/UnrealString.h"
+
+struct FScriptContainerElement;
 
 namespace GLTF
 {
@@ -89,3 +91,7 @@ namespace GLTF
 	};
 
 }  // namespace GLTF
+
+#if UE_ENABLE_INCLUDE_ORDER_DEPRECATED_IN_5_2
+#include "CoreMinimal.h"
+#endif
diff --git a/Engine/Plugins/Editor/GLTFImporter/Source/GLTFImporter/Public/GLTFImportOptions.h b/Engine/Plugins/Editor/GLTFImporter/Source/GLTFImporter/Public/GLTFImportOptions.h
--- a/Engine/Plugins/Editor/GLTFImporter/Source/GLTFImporter/Public/GLTFImportOptions.h
+++ b/Engine/Plugins/Editor/GLTFImporter/Source/GLTFImporter/Public/GLTFImportOptions.h
@@ -2,10 +2,6 @@
 
 #pragma once
 
-#include "CoreMinimal.h"
-#include "Engine/EngineTypes.h"
-#include "UObject/Object.h"
-#include "UObject/ObjectMacros.h"
 
 #include "GLTFImportOptions.generated.h"
 
@@ -31,3 +27,8 @@ public:
 	         ToolTip = "Scale factor used for importing assets, by default: 100, for conversion from meters(glTF) to centimeters(Unreal default)."))
 	float ImportScale;
 };
+
+#if UE_ENABLE_INCLUDE_ORDER_DEPRECATED_IN_5_2
+#include "CoreMinimal.h"
+#include "Engine/EngineTypes.h"
+#endif
diff --git a/Engine/Plugins/Editor/GLTFImporter/Source/GLTFImporter/Public/GLTFImporterModule.h b/Engine/Plugins/Editor/GLTFImporter/Source/GLTFImporter/Public/GLTFImporterModule.h
--- a/Engine/Plugins/Editor/GLTFImporter/Source/GLTFImporter/Public/GLTFImporterModule.h
+++ b/Engine/Plugins/Editor/GLTFImporter/Source/GLTFImporter/Public/GLTFImporterModule.h
@@ -2,8 +2,6 @@
 
 #pragma once
 
-#include "CoreMinimal.h"
-#include "Modules/ModuleInterface.h"
 #include "Modules/ModuleManager.h"
 
 struct FGLTFImporterContext;
@@ -41,3 +39,7 @@ public:
 		return FModuleManager::Get().IsModuleLoaded("GLTFImporter");
 	}
 };
+
+#if UE_ENABLE_INCLUDE_ORDER_DEPRECATED_IN_5_2
+#include "CoreMinimal.h"
+#endif
diff --git a/Engine/Plugins/Editor/GLTFImporter/Source/GLTFImporter/Public/GLTFStaticMeshFactory.h b/Engine/Plugins/Editor/GLTFImporter/Source/GLTFImporter/Public/GLTFStaticMeshFactory.h
--- a/Engine/Plugins/Editor/GLTFImporter/Source/GLTFImporter/Public/GLTFStaticMeshFactory.h
+++ b/Engine/Plugins/Editor/GLTFImporter/Source/GLTFImporter/Public/GLTFStaticMeshFactory.h
@@ -4,8 +4,6 @@
 
 #include "GLTFLogger.h"
 
-#include "CoreMinimal.h"
-#include "Templates/UniquePtr.h"
 #include "UObject/ObjectMacros.h"
 
 class UStaticMesh;
@@ -45,3 +43,7 @@ namespace GLTF
 	};
 
 }  // namespace GLTF
+
+#if UE_ENABLE_INCLUDE_ORDER_DEPRECATED_IN_5_2
+#include "CoreMinimal.h"
+#endif
diff --git a/Engine/Plugins/Editor/ModelingToolsEditorMode/Source/ModelingToolsEditorMode/Private/ModelingModeToolkit_Toolbars.cpp b/Engine/Plugins/Editor/ModelingToolsEditorMode/Source/ModelingToolsEditorMode/Private/ModelingModeToolkit_Toolbars.cpp
--- a/Engine/Plugins/Editor/ModelingToolsEditorMode/Source/ModelingToolsEditorMode/Private/ModelingModeToolkit_Toolbars.cpp
+++ b/Engine/Plugins/Editor/ModelingToolsEditorMode/Source/ModelingToolsEditorMode/Private/ModelingModeToolkit_Toolbars.cpp
@@ -18,8 +18,6 @@
 
 #include "Framework/MultiBox/MultiBoxBuilder.h"
 #include "STransformGizmoNumericalUIOverlay.h"
-#include "Widgets/Images/SImage.h"
-#include "Widgets/Layout/SBorder.h"
 #include "Widgets/Layout/SSpacer.h"
 #include "Widgets/Input/SComboButton.h"
 
diff --git a/Engine/Plugins/Editor/ModelingToolsEditorMode/Source/ModelingToolsEditorMode/Private/ModelingToolsEditorMode.cpp b/Engine/Plugins/Editor/ModelingToolsEditorMode/Source/ModelingToolsEditorMode/Private/ModelingToolsEditorMode.cpp
--- a/Engine/Plugins/Editor/ModelingToolsEditorMode/Source/ModelingToolsEditorMode/Private/ModelingToolsEditorMode.cpp
+++ b/Engine/Plugins/Editor/ModelingToolsEditorMode/Source/ModelingToolsEditorMode/Private/ModelingToolsEditorMode.cpp
@@ -20,7 +20,6 @@
 #include "ToolTargets/SkeletalMeshComponentToolTarget.h"
 #include "EngineAnalytics.h"
 #include "BaseGizmos/TransformGizmoUtil.h"
-#include "Selection/StoredMeshSelectionUtil.h"
 #include "Snapping/ModelingSceneSnappingManager.h"
 #include "Scene/LevelObjectsObserver.h"
 
diff --git a/Engine/Plugins/Editor/ModelingToolsEditorMode/Source/ModelingToolsEditorMode/Public/DetailsCustomizations/BakeMeshAttributeToolCustomizations.h b/Engine/Plugins/Editor/ModelingToolsEditorMode/Source/ModelingToolsEditorMode/Public/DetailsCustomizations/BakeMeshAttributeToolCustomizations.h
--- a/Engine/Plugins/Editor/ModelingToolsEditorMode/Source/ModelingToolsEditorMode/Public/DetailsCustomizations/BakeMeshAttributeToolCustomizations.h
+++ b/Engine/Plugins/Editor/ModelingToolsEditorMode/Source/ModelingToolsEditorMode/Public/DetailsCustomizations/BakeMeshAttributeToolCustomizations.h
@@ -3,7 +3,6 @@
 #pragma once
 
 #include "IDetailCustomization.h"
-#include "Templates/SharedPointer.h"
 
 // Customization for UBakeMeshAttributeMapsToolProperties
 class FBakeMeshAttributeMapsToolDetails : public IDetailCustomization
diff --git a/Engine/Plugins/Editor/ModelingToolsEditorMode/Source/ModelingToolsEditorMode/Public/DetailsCustomizations/BakeTransformToolCustomizations.h b/Engine/Plugins/Editor/ModelingToolsEditorMode/Source/ModelingToolsEditorMode/Public/DetailsCustomizations/BakeTransformToolCustomizations.h
--- a/Engine/Plugins/Editor/ModelingToolsEditorMode/Source/ModelingToolsEditorMode/Public/DetailsCustomizations/BakeTransformToolCustomizations.h
+++ b/Engine/Plugins/Editor/ModelingToolsEditorMode/Source/ModelingToolsEditorMode/Public/DetailsCustomizations/BakeTransformToolCustomizations.h
@@ -3,7 +3,6 @@
 #pragma once
 
 #include "IDetailCustomization.h"
-#include "Templates/SharedPointer.h"
 
 class IPropertyHandle;
 class FPropertyRestriction;
diff --git a/Engine/Plugins/Editor/ModelingToolsEditorMode/Source/ModelingToolsEditorMode/Public/DetailsCustomizations/MeshVertexSculptToolCustomizations.h b/Engine/Plugins/Editor/ModelingToolsEditorMode/Source/ModelingToolsEditorMode/Public/DetailsCustomizations/MeshVertexSculptToolCustomizations.h
--- a/Engine/Plugins/Editor/ModelingToolsEditorMode/Source/ModelingToolsEditorMode/Public/DetailsCustomizations/MeshVertexSculptToolCustomizations.h
+++ b/Engine/Plugins/Editor/ModelingToolsEditorMode/Source/ModelingToolsEditorMode/Public/DetailsCustomizations/MeshVertexSculptToolCustomizations.h
@@ -3,10 +3,12 @@
 #pragma once
 
 #include "IDetailCustomization.h"
-#include "Templates/SharedPointer.h"
-#include "IPropertyTypeCustomization.h"
-#include "Widgets/Input/SComboBox.h"
-#include "Widgets/Views/STileView.h"
+#include "UObject/WeakObjectPtr.h"
+
+class FReply;
+class SButton;
+class SWidget;
+enum class ECheckBoxState : uint8;
 
 class IDetailLayoutBuilder;
 class IPropertyHandle;
@@ -58,4 +60,10 @@ public:
 	void CustomizeDetails(IDetailLayoutBuilder& DetailBuilder) override;
 
 	TSharedPtr<FRecentAlphasProvider> RecentAlphasProvider;
-};
\ No newline at end of file
+};
+
+#if UE_ENABLE_INCLUDE_ORDER_DEPRECATED_IN_5_2
+#include "IPropertyTypeCustomization.h"
+#include "Widgets/Input/SComboBox.h"
+#include "Widgets/Views/STileView.h"
+#endif
diff --git a/Engine/Plugins/Editor/ModelingToolsEditorMode/Source/ModelingToolsEditorMode/Public/DetailsCustomizations/ModelingToolPropertyCustomizations.h b/Engine/Plugins/Editor/ModelingToolsEditorMode/Source/ModelingToolsEditorMode/Public/DetailsCustomizations/ModelingToolPropertyCustomizations.h
--- a/Engine/Plugins/Editor/ModelingToolsEditorMode/Source/ModelingToolsEditorMode/Public/DetailsCustomizations/ModelingToolPropertyCustomizations.h
+++ b/Engine/Plugins/Editor/ModelingToolsEditorMode/Source/ModelingToolsEditorMode/Public/DetailsCustomizations/ModelingToolPropertyCustomizations.h
@@ -3,9 +3,6 @@
 #pragma once
 
 #include "IPropertyTypeCustomization.h"
-#include "Input/Reply.h"
-#include "Widgets/Text/STextBlock.h"
-#include "Widgets/Input/SButton.h"
 
 class IDetailLayoutBuilder;
 
@@ -18,3 +15,9 @@ public:
 	virtual void CustomizeHeader(TSharedRef<IPropertyHandle> StructPropertyHandle, class FDetailWidgetRow& HeaderRow, IPropertyTypeCustomizationUtils& StructCustomizationUtils) override;
 	virtual void CustomizeChildren(TSharedRef<IPropertyHandle> StructPropertyHandle, class IDetailChildrenBuilder& ChildBuilder, IPropertyTypeCustomizationUtils& StructCustomizationUtils) override;
 };
+
+#if UE_ENABLE_INCLUDE_ORDER_DEPRECATED_IN_5_2
+#include "Input/Reply.h"
+#include "Widgets/Input/SButton.h"
+#include "Widgets/Text/STextBlock.h"
+#endif
diff --git a/Engine/Plugins/Editor/ModelingToolsEditorMode/Source/ModelingToolsEditorMode/Public/ModelingModeAssetUtils.h b/Engine/Plugins/Editor/ModelingToolsEditorMode/Source/ModelingToolsEditorMode/Public/ModelingModeAssetUtils.h
--- a/Engine/Plugins/Editor/ModelingToolsEditorMode/Source/ModelingToolsEditorMode/Public/ModelingModeAssetUtils.h
+++ b/Engine/Plugins/Editor/ModelingToolsEditorMode/Source/ModelingToolsEditorMode/Public/ModelingModeAssetUtils.h
@@ -2,7 +2,10 @@
 
 #pragma once
 
-#include "CoreMinimal.h"
+#include "HAL/Platform.h"
+
+class FString;
+class UObject;
 
 class UWorld;
 
@@ -54,4 +57,8 @@ MODELINGTOOLSEDITORMODE_API void OnNewAssetCreated(UObject* Asset);
 MODELINGTOOLSEDITORMODE_API bool AutoSaveAsset(UObject* Asset);
 
 }
-}
\ No newline at end of file
+}
+
+#if UE_ENABLE_INCLUDE_ORDER_DEPRECATED_IN_5_2
+#include "CoreMinimal.h"
+#endif
diff --git a/Engine/Plugins/Editor/ModelingToolsEditorMode/Source/ModelingToolsEditorMode/Public/ModelingModeToolExtensions.h b/Engine/Plugins/Editor/ModelingToolsEditorMode/Source/ModelingToolsEditorMode/Public/ModelingModeToolExtensions.h
--- a/Engine/Plugins/Editor/ModelingToolsEditorMode/Source/ModelingToolsEditorMode/Public/ModelingModeToolExtensions.h
+++ b/Engine/Plugins/Editor/ModelingToolsEditorMode/Source/ModelingToolsEditorMode/Public/ModelingModeToolExtensions.h
@@ -2,8 +2,8 @@
 
 #pragma once
 
-#include "CoreMinimal.h"
 #include "Features/IModularFeature.h"
+#include "Internationalization/Text.h"
 
 class UInteractiveToolsContext;
 class FUICommandInfo;
@@ -90,3 +90,7 @@ public:
 };
 
 
+
+#if UE_ENABLE_INCLUDE_ORDER_DEPRECATED_IN_5_2
+#include "CoreMinimal.h"
+#endif
diff --git a/Engine/Plugins/Editor/ModelingToolsEditorMode/Source/ModelingToolsEditorMode/Public/ModelingSelectionInteraction.h b/Engine/Plugins/Editor/ModelingToolsEditorMode/Source/ModelingToolsEditorMode/Public/ModelingSelectionInteraction.h
--- a/Engine/Plugins/Editor/ModelingToolsEditorMode/Source/ModelingToolsEditorMode/Public/ModelingSelectionInteraction.h
+++ b/Engine/Plugins/Editor/ModelingToolsEditorMode/Source/ModelingToolsEditorMode/Public/ModelingSelectionInteraction.h
@@ -5,13 +5,16 @@
 
 #include "BaseBehaviors/BehaviorTargetInterfaces.h"
 #include "InputBehaviorSet.h"
-#include "BaseGizmos/TransformGizmoUtil.h"
 #include "FrameTypes.h"
-#include "Mechanics/RectangleMarqueeMechanic.h"
-#include "Mechanics/PolyLassoMarqueeMechanic.h"
-#include "Selections/GeometrySelection.h"
+#include "UObject/Package.h"
 #include "ModelingSelectionInteraction.generated.h"
 
+class UCombinedTransformGizmo;
+class URectangleMarqueeInteraction;
+class UTransformProxy;
+namespace UE::Geometry { struct FGeometrySelectionUpdateConfig; }
+struct FCameraRectangle;
+
 class UGeometrySelectionManager;
 class USingleClickInputBehavior;
 class USingleClickOrDragInputBehavior;
@@ -238,4 +241,11 @@ protected:
 	virtual void OnClickRelease(const FInputDeviceRay& ReleasePos) override;
 	virtual void OnTerminateDragSequence() override;
 
-};
\ No newline at end of file
+};
+
+#if UE_ENABLE_INCLUDE_ORDER_DEPRECATED_IN_5_2
+#include "BaseGizmos/TransformGizmoUtil.h"
+#include "Mechanics/PolyLassoMarqueeMechanic.h"
+#include "Mechanics/RectangleMarqueeMechanic.h"
+#include "Selections/GeometrySelection.h"
+#endif
diff --git a/Engine/Plugins/Editor/ModelingToolsEditorMode/Source/ModelingToolsEditorMode/Public/ModelingToolsEditorMode.h b/Engine/Plugins/Editor/ModelingToolsEditorMode/Source/ModelingToolsEditorMode/Public/ModelingToolsEditorMode.h
--- a/Engine/Plugins/Editor/ModelingToolsEditorMode/Source/ModelingToolsEditorMode/Public/ModelingToolsEditorMode.h
+++ b/Engine/Plugins/Editor/ModelingToolsEditorMode/Source/ModelingToolsEditorMode/Public/ModelingToolsEditorMode.h
@@ -2,16 +2,17 @@
 
 #pragma once
 
-#include "CoreMinimal.h"
 #include "Tools/LegacyEdModeWidgetHelpers.h"
 
-#include "InputState.h"
-#include "InteractiveToolManager.h"
-#include "EdModeInteractiveToolsContext.h"
-#include "ModelingToolsActions.h"
 
 #include "ModelingToolsEditorMode.generated.h"
 
+class IToolsContextRenderAPI;
+enum class EModelingModeActionCommands;
+enum class EToolSide;
+struct FInputDeviceRay;
+struct FToolBuilderState;
+
 class FEditorComponentSourceFactory;
 class FUICommandList;
 class FStylusStateTracker;		// for stylus events
@@ -151,3 +152,11 @@ protected:
 	static FDateTime LastModeStartTimestamp;
 	static FDateTime LastToolStartTimestamp;
 };
+
+#if UE_ENABLE_INCLUDE_ORDER_DEPRECATED_IN_5_2
+#include "CoreMinimal.h"
+#include "EdModeInteractiveToolsContext.h"
+#include "InputState.h"
+#include "InteractiveToolManager.h"
+#include "ModelingToolsActions.h"
+#endif
diff --git a/Engine/Plugins/Editor/ModelingToolsEditorMode/Source/ModelingToolsEditorMode/Public/ModelingToolsEditorModeModule.h b/Engine/Plugins/Editor/ModelingToolsEditorMode/Source/ModelingToolsEditorMode/Public/ModelingToolsEditorModeModule.h
--- a/Engine/Plugins/Editor/ModelingToolsEditorMode/Source/ModelingToolsEditorMode/Public/ModelingToolsEditorModeModule.h
+++ b/Engine/Plugins/Editor/ModelingToolsEditorMode/Source/ModelingToolsEditorMode/Public/ModelingToolsEditorModeModule.h
@@ -2,8 +2,8 @@
 
 #pragma once
 
-#include "CoreMinimal.h"
-#include "Modules/ModuleManager.h"
+#include "Containers/Array.h"
+#include "Modules/ModuleInterface.h"
 
 class FModelingToolsEditorModeModule : public IModuleInterface
 {
@@ -20,3 +20,8 @@ private:
 	TArray<FName> ClassesToUnregisterOnShutdown;
 	TArray<FName> PropertiesToUnregisterOnShutdown;
 };
+
+#if UE_ENABLE_INCLUDE_ORDER_DEPRECATED_IN_5_2
+#include "CoreMinimal.h"
+#include "Modules/ModuleManager.h"
+#endif
diff --git a/Engine/Plugins/Editor/ModelingToolsEditorMode/Source/ModelingToolsEditorMode/Public/ModelingToolsEditorModeSettings.h b/Engine/Plugins/Editor/ModelingToolsEditorMode/Source/ModelingToolsEditorMode/Public/ModelingToolsEditorModeSettings.h
--- a/Engine/Plugins/Editor/ModelingToolsEditorMode/Source/ModelingToolsEditorMode/Public/ModelingToolsEditorModeSettings.h
+++ b/Engine/Plugins/Editor/ModelingToolsEditorMode/Source/ModelingToolsEditorMode/Public/ModelingToolsEditorModeSettings.h
@@ -3,9 +3,10 @@
 #pragma once
 
 #include "Engine/DeveloperSettings.h"
-#include "Engine/EngineTypes.h"
 #include "ModelingToolsEditorModeSettings.generated.h"
 
+struct FCollectionReference;
+
 
 
 UENUM()
@@ -353,3 +354,7 @@ public:
 
 
 };
+
+#if UE_ENABLE_INCLUDE_ORDER_DEPRECATED_IN_5_2
+#include "Engine/EngineTypes.h"
+#endif
diff --git a/Engine/Plugins/Editor/ModelingToolsEditorMode/Source/ModelingToolsEditorMode/Public/ModelingToolsEditorModeStyle.h b/Engine/Plugins/Editor/ModelingToolsEditorMode/Source/ModelingToolsEditorMode/Public/ModelingToolsEditorModeStyle.h
--- a/Engine/Plugins/Editor/ModelingToolsEditorMode/Source/ModelingToolsEditorMode/Public/ModelingToolsEditorModeStyle.h
+++ b/Engine/Plugins/Editor/ModelingToolsEditorMode/Source/ModelingToolsEditorMode/Public/ModelingToolsEditorModeStyle.h
@@ -2,9 +2,9 @@
 
 #pragma once
 
-#include "CoreMinimal.h"
-#include "Styling/ISlateStyle.h"
-#include "Styling/SlateStyle.h"
+#include "Templates/SharedPointer.h"
+
+struct FSlateBrush;
 
 class MODELINGTOOLSEDITORMODE_API FModelingToolsEditorModeStyle
 {
@@ -26,3 +26,9 @@ private:
 private:
 	static TSharedPtr< class FSlateStyleSet > StyleSet;
 };
+
+#if UE_ENABLE_INCLUDE_ORDER_DEPRECATED_IN_5_2
+#include "CoreMinimal.h"
+#include "Styling/ISlateStyle.h"
+#include "Styling/SlateStyle.h"
+#endif
diff --git a/Engine/Plugins/Editor/ModelingToolsEditorMode/Source/ModelingToolsEditorMode/Public/ModelingToolsEditorModeToolkit.h b/Engine/Plugins/Editor/ModelingToolsEditorMode/Source/ModelingToolsEditorMode/Public/ModelingToolsEditorModeToolkit.h
--- a/Engine/Plugins/Editor/ModelingToolsEditorMode/Source/ModelingToolsEditorMode/Public/ModelingToolsEditorModeToolkit.h
+++ b/Engine/Plugins/Editor/ModelingToolsEditorMode/Source/ModelingToolsEditorMode/Public/ModelingToolsEditorModeToolkit.h
@@ -2,14 +2,15 @@
 
 #pragma once
 
-#include "CoreMinimal.h"
 #include "Toolkits/BaseToolkit.h"
-#include "ModelingToolsEditorMode.h"
-#include "Widgets/Input/STextComboBox.h"
-#include "Widgets/Input/SEditableTextBox.h"
 #include "StatusBarSubsystem.h"
 #include "UObject/GCObject.h"
 
+class IAssetViewport;
+class SEditableTextBox;
+class STextComboBox;
+namespace ETextCommit { enum Type : int; }
+
 #include "ModelingToolsEditorModeToolkit.generated.h"
 
 #define LOCTEXT_NAMESPACE "ModelingToolsEditorModeToolkit"
@@ -183,4 +184,8 @@ private:
 	bool bShowActiveSelectionActions = true;
 };
 
-#undef LOCTEXT_NAMESPACE
\ No newline at end of file
+#include "CoreMinimal.h"
+#include "ModelingToolsEditorMode.h"
+#include "Widgets/Input/SEditableTextBox.h"
+#include "Widgets/Input/STextComboBox.h"
+#endif
diff --git a/Engine/Plugins/EnhancedInput/Source/EnhancedInput/Private/EnhancedInputWorldProcessor.cpp b/Engine/Plugins/EnhancedInput/Source/EnhancedInput/Private/EnhancedInputWorldProcessor.cpp
--- a/Engine/Plugins/EnhancedInput/Source/EnhancedInput/Private/EnhancedInputWorldProcessor.cpp
+++ b/Engine/Plugins/EnhancedInput/Source/EnhancedInput/Private/EnhancedInputWorldProcessor.cpp
@@ -2,11 +2,9 @@
 
 #include "EnhancedInputWorldProcessor.h"
 #include "EnhancedInputSubsystems.h"
-#include "Engine/World.h"
 #include "Framework/Application/SlateApplication.h"
 #include "EnhancedInputDeveloperSettings.h"
 #include "UObject/UObjectIterator.h"
-#include "GenericPlatform/GenericPlatformInputDeviceMapper.h"
 
 void FEnhancedInputWorldProcessor::Tick(const float DeltaTime, FSlateApplication& SlateApp, TSharedRef<ICursor> Cursor)
 {
diff --git a/Engine/Plugins/Experimental/Animation/ContextualAnimation/Source/ContextualAnimation/Public/ContextualAnimManager.h b/Engine/Plugins/Experimental/Animation/ContextualAnimation/Source/ContextualAnimation/Public/ContextualAnimManager.h
--- a/Engine/Plugins/Experimental/Animation/ContextualAnimation/Source/ContextualAnimation/Public/ContextualAnimManager.h
+++ b/Engine/Plugins/Experimental/Animation/ContextualAnimation/Source/ContextualAnimation/Public/ContextualAnimManager.h
@@ -2,12 +2,11 @@
 
 #pragma once
 
-#include "CoreMinimal.h"
-#include "UObject/Object.h"
 #include "Tickable.h"
-#include "ContextualAnimTypes.h"
 #include "ContextualAnimManager.generated.h"
 
+struct FContextualAnimStartSceneParams;
+
 class UContextualAnimSceneActorComponent;
 class UContextualAnimSceneAsset;
 class UContextualAnimSceneInstance;
@@ -76,4 +75,9 @@ protected:
 
 	UFUNCTION()
 	void OnSceneInstanceEnded(UContextualAnimSceneInstance* SceneInstance);
-};
\ No newline at end of file
+};
+
+#if UE_ENABLE_INCLUDE_ORDER_DEPRECATED_IN_5_2
+#include "ContextualAnimTypes.h"
+#include "CoreMinimal.h"
+#endif
diff --git a/Engine/Plugins/Experimental/Animation/ContextualAnimation/Source/ContextualAnimation/Public/ContextualAnimSceneAsset.h b/Engine/Plugins/Experimental/Animation/ContextualAnimation/Source/ContextualAnimation/Public/ContextualAnimSceneAsset.h
--- a/Engine/Plugins/Experimental/Animation/ContextualAnimation/Source/ContextualAnimation/Public/ContextualAnimSceneAsset.h
+++ b/Engine/Plugins/Experimental/Animation/ContextualAnimation/Source/ContextualAnimation/Public/ContextualAnimSceneAsset.h
@@ -3,7 +3,6 @@
 #pragma once
 
 #include "Engine/DataAsset.h"
-#include "Templates/SubclassOf.h"
 #include "ContextualAnimTypes.h"
 #include "ContextualAnimSceneAsset.generated.h"
 
diff --git a/Engine/Plugins/Experimental/Animation/ContextualAnimation/Source/ContextualAnimation/Public/ContextualAnimSceneInstance.h b/Engine/Plugins/Experimental/Animation/ContextualAnimation/Source/ContextualAnimation/Public/ContextualAnimSceneInstance.h
--- a/Engine/Plugins/Experimental/Animation/ContextualAnimation/Source/ContextualAnimation/Public/ContextualAnimSceneInstance.h
+++ b/Engine/Plugins/Experimental/Animation/ContextualAnimation/Source/ContextualAnimation/Public/ContextualAnimSceneInstance.h
@@ -2,13 +2,11 @@
 
 #pragma once
 
-#include "CoreMinimal.h"
-#include "UObject/Object.h"
-#include "Animation/AnimNotifies/AnimNotify.h"
-#include "GameFramework/Actor.h"
 #include "ContextualAnimTypes.h"
 #include "ContextualAnimSceneInstance.generated.h"
 
+struct FBranchingPointNotifyPayload;
+
 struct FAnimMontageInstance;
 struct FContextualAnimTrack;
 class UContextualAnimSceneInstance;
@@ -152,4 +150,9 @@ private:
 
 	/** Helper to play an AnimSequenceBase as montage. If Animation is not a montage it plays it as dynamic montage  */
 	UAnimMontage* PlayAnimation(UAnimInstance& AnimInstance, UAnimSequenceBase& Animation);
-};
\ No newline at end of file
+};
+
+#if UE_ENABLE_INCLUDE_ORDER_DEPRECATED_IN_5_2
+#include "Animation/AnimNotifies/AnimNotify.h"
+#include "CoreMinimal.h"
+#endif
diff --git a/Engine/Plugins/Experimental/Animation/ContextualAnimation/Source/ContextualAnimation/Public/ContextualAnimSelectionCriterion.h b/Engine/Plugins/Experimental/Animation/ContextualAnimation/Source/ContextualAnimation/Public/ContextualAnimSelectionCriterion.h
--- a/Engine/Plugins/Experimental/Animation/ContextualAnimation/Source/ContextualAnimation/Public/ContextualAnimSelectionCriterion.h
+++ b/Engine/Plugins/Experimental/Animation/ContextualAnimation/Source/ContextualAnimation/Public/ContextualAnimSelectionCriterion.h
@@ -2,11 +2,11 @@
 
 #pragma once
 
-#include "CoreMinimal.h"
-#include "UObject/Object.h"
-#include "ContextualAnimTypes.h"
 #include "ContextualAnimSelectionCriterion.generated.h"
 
+class UContextualAnimSceneAsset;
+struct FContextualAnimSceneBindingContext;
+
 UENUM(BlueprintType)
 enum class EContextualAnimCriterionType : uint8
 {
@@ -142,4 +142,9 @@ public:
 	UContextualAnimSelectionCriterion_Distance(const FObjectInitializer& ObjectInitializer) : Super(ObjectInitializer) {}
 
 	virtual bool DoesQuerierPassCondition(const FContextualAnimSceneBindingContext& Primary, const FContextualAnimSceneBindingContext& Querier) const override;
-};
\ No newline at end of file
+};
+
+#if UE_ENABLE_INCLUDE_ORDER_DEPRECATED_IN_5_2
+#include "ContextualAnimTypes.h"
+#include "CoreMinimal.h"
+#endif
diff --git a/Engine/Plugins/Experimental/Animation/ContextualAnimation/Source/ContextualAnimation/Public/ContextualAnimTransition.h b/Engine/Plugins/Experimental/Animation/ContextualAnimation/Source/ContextualAnimation/Public/ContextualAnimTransition.h
--- a/Engine/Plugins/Experimental/Animation/ContextualAnimation/Source/ContextualAnimation/Public/ContextualAnimTransition.h
+++ b/Engine/Plugins/Experimental/Animation/ContextualAnimation/Source/ContextualAnimation/Public/ContextualAnimTransition.h
@@ -2,7 +2,6 @@
 
 #pragma once
 
-#include "UObject/Object.h"
 #include "ContextualAnimTransition.generated.h"
 
 class AActor;
@@ -19,4 +18,4 @@ public:
 
 	UFUNCTION(BlueprintImplementableEvent, Category = "Defaults")
 	bool CanEnterTransition(const UContextualAnimSceneInstance* SceneInstance, const FName& FromSection, const FName& ToSection) const;
-};
\ No newline at end of file
+};
diff --git a/Engine/Plugins/Experimental/Animation/ContextualAnimation/Source/ContextualAnimation/Public/ContextualAnimTypes.h b/Engine/Plugins/Experimental/Animation/ContextualAnimation/Source/ContextualAnimation/Public/ContextualAnimTypes.h
--- a/Engine/Plugins/Experimental/Animation/ContextualAnimation/Source/ContextualAnimation/Public/ContextualAnimTypes.h
+++ b/Engine/Plugins/Experimental/Animation/ContextualAnimation/Source/ContextualAnimation/Public/ContextualAnimTypes.h
@@ -3,11 +3,12 @@
 #pragma once
 
 #include "Animation/AnimSequence.h"
-#include "Templates/SubclassOf.h"
 #include "GameFramework/Actor.h"
 #include "GameplayTagContainer.h"
 #include "ContextualAnimTypes.generated.h"
 
+struct FContextualAnimSceneBindingContext;
+
 CONTEXTUALANIMATION_API DECLARE_LOG_CATEGORY_EXTERN(LogContextualAnim, Log, All);
 
 class AActor;
diff --git a/Engine/Plugins/Experimental/Animation/ContextualAnimation/Source/ContextualAnimation/Public/ContextualAnimUtilities.h b/Engine/Plugins/Experimental/Animation/ContextualAnimation/Source/ContextualAnimation/Public/ContextualAnimUtilities.h
--- a/Engine/Plugins/Experimental/Animation/ContextualAnimation/Source/ContextualAnimation/Public/ContextualAnimUtilities.h
+++ b/Engine/Plugins/Experimental/Animation/ContextualAnimation/Source/ContextualAnimation/Public/ContextualAnimUtilities.h
@@ -2,8 +2,6 @@
 
 #pragma once
 
-#include "Kismet/BlueprintFunctionLibrary.h"
-#include "Animation/AnimSequence.h"
 #include "ContextualAnimTypes.h"
 #include "ContextualAnimUtilities.generated.h"
 
@@ -146,4 +144,4 @@ public:
 	UFUNCTION(BlueprintPure, Category = "Contextual Anim|Scene Binding", meta = (DisplayName = "Get Skeletal Mesh"))
 	static USkeletalMeshComponent* BP_SceneBinding_GetSkeletalMesh(const FContextualAnimSceneBinding& Binding) { return Binding.GetSkeletalMeshComponent(); }
 
-};
\ No newline at end of file
+};
diff --git a/Engine/Plugins/Experimental/Animation/ContextualAnimation/Source/ContextualAnimation/Public/ContextualAnimation.h b/Engine/Plugins/Experimental/Animation/ContextualAnimation/Source/ContextualAnimation/Public/ContextualAnimation.h
--- a/Engine/Plugins/Experimental/Animation/ContextualAnimation/Source/ContextualAnimation/Public/ContextualAnimation.h
+++ b/Engine/Plugins/Experimental/Animation/ContextualAnimation/Source/ContextualAnimation/Public/ContextualAnimation.h
@@ -2,10 +2,8 @@
 
 #pragma once
 
-#include "CoreMinimal.h"
-#include "UObject/GCObject.h"
-#include "Modules/ModuleManager.h"
 #include "Engine/World.h"
+#include "Modules/ModuleInterface.h"
 
 class UContextualAnimManager;
 
@@ -41,3 +39,8 @@ private:
 	FDelegateHandle OnPreWorldInitDelegateHandle;
 	FDelegateHandle OnPostWorldCleanupDelegateHandle;
 };
+
+#if UE_ENABLE_INCLUDE_ORDER_DEPRECATED_IN_5_2
+#include "CoreMinimal.h"
+#include "Modules/ModuleManager.h"
+#endif
diff --git a/Engine/Plugins/Experimental/Animation/ContextualAnimation/Source/ContextualAnimationEditor/Public/ContextualAnimAssetEditorToolkit.h b/Engine/Plugins/Experimental/Animation/ContextualAnimation/Source/ContextualAnimationEditor/Public/ContextualAnimAssetEditorToolkit.h
--- a/Engine/Plugins/Experimental/Animation/ContextualAnimation/Source/ContextualAnimationEditor/Public/ContextualAnimAssetEditorToolkit.h
+++ b/Engine/Plugins/Experimental/Animation/ContextualAnimation/Source/ContextualAnimationEditor/Public/ContextualAnimAssetEditorToolkit.h
@@ -2,9 +2,9 @@
 
 #pragma once
 
-#include "CoreMinimal.h"
 #include "Toolkits/AssetEditorToolkit.h"
-#include "Misc/NotifyHook.h"
+
+class FSpawnTabArgs;
 
 class IDetailsView;
 class FToolBarBuilder;
@@ -68,3 +68,8 @@ private:
 };
 
 
+
+#if UE_ENABLE_INCLUDE_ORDER_DEPRECATED_IN_5_2
+#include "CoreMinimal.h"
+#include "Misc/NotifyHook.h"
+#endif
diff --git a/Engine/Plugins/Experimental/Animation/ContextualAnimation/Source/ContextualAnimationEditor/Public/ContextualAnimEdMode.h b/Engine/Plugins/Experimental/Animation/ContextualAnimation/Source/ContextualAnimationEditor/Public/ContextualAnimEdMode.h
--- a/Engine/Plugins/Experimental/Animation/ContextualAnimation/Source/ContextualAnimationEditor/Public/ContextualAnimEdMode.h
+++ b/Engine/Plugins/Experimental/Animation/ContextualAnimation/Source/ContextualAnimationEditor/Public/ContextualAnimEdMode.h
@@ -2,8 +2,8 @@
 
 #pragma once
 
-#include "CoreMinimal.h"
 #include "EdMode.h"
+#include "GenericPlatform/ICursor.h"
 #include "HitProxies.h"
 
 class AActor;
@@ -58,3 +58,7 @@ private:
 
 	FContextualAnimViewModel* ViewModel = nullptr;
 };
+
+#if UE_ENABLE_INCLUDE_ORDER_DEPRECATED_IN_5_2
+#include "CoreMinimal.h"
+#endif
diff --git a/Engine/Plugins/Experimental/Animation/ContextualAnimation/Source/ContextualAnimationEditor/Public/ContextualAnimEditorStyle.h b/Engine/Plugins/Experimental/Animation/ContextualAnimation/Source/ContextualAnimationEditor/Public/ContextualAnimEditorStyle.h
--- a/Engine/Plugins/Experimental/Animation/ContextualAnimation/Source/ContextualAnimationEditor/Public/ContextualAnimEditorStyle.h
+++ b/Engine/Plugins/Experimental/Animation/ContextualAnimation/Source/ContextualAnimationEditor/Public/ContextualAnimEditorStyle.h
@@ -2,8 +2,9 @@
 
 #pragma once
 
-#include "CoreMinimal.h"
-#include "Styling/SlateStyle.h"
+#include "Templates/SharedPointer.h"
+
+class ISlateStyle;
 
 class FContextualAnimEditorStyle
 {
@@ -26,4 +27,9 @@ private:
 	static TSharedRef< class FSlateStyleSet > Create();
 
 	static TSharedPtr< class FSlateStyleSet > StyleInstance;
-};
\ No newline at end of file
+};
+
+#if UE_ENABLE_INCLUDE_ORDER_DEPRECATED_IN_5_2
+#include "CoreMinimal.h"
+#include "Styling/SlateStyle.h"
+#endif
diff --git a/Engine/Plugins/Experimental/Animation/ContextualAnimation/Source/ContextualAnimationEditor/Public/ContextualAnimEditorTypes.h b/Engine/Plugins/Experimental/Animation/ContextualAnimation/Source/ContextualAnimationEditor/Public/ContextualAnimEditorTypes.h
--- a/Engine/Plugins/Experimental/Animation/ContextualAnimation/Source/ContextualAnimationEditor/Public/ContextualAnimEditorTypes.h
+++ b/Engine/Plugins/Experimental/Animation/ContextualAnimation/Source/ContextualAnimationEditor/Public/ContextualAnimEditorTypes.h
@@ -2,11 +2,12 @@
 
 #pragma once
 
-#include "CoreMinimal.h"
 #include "ContextualAnimTypes.h"
 #include "Interfaces/Interface_BoneReferenceSkeletonProvider.h"
 #include "ContextualAnimEditorTypes.generated.h"
 
+class IPropertyHandle;
+
 class UContextualAnimSceneAsset;
 
 /** Object used to construct the New IK Target Widget */
@@ -96,4 +97,8 @@ struct FContextualAnimNewAnimSetParams
 
 	UPROPERTY(EditAnywhere, Category = "Settings", meta = (ClampMin = "0", UIMin = "0", ClampMax = "1", UIMax = "1"))
 	float RandomWeight = 1.f;
-};
\ No newline at end of file
+};
+
+#if UE_ENABLE_INCLUDE_ORDER_DEPRECATED_IN_5_2
+#include "CoreMinimal.h"
+#endif
diff --git a/Engine/Plugins/Experimental/Animation/ContextualAnimation/Source/ContextualAnimationEditor/Public/ContextualAnimMovieSceneNotifyTrack.h b/Engine/Plugins/Experimental/Animation/ContextualAnimation/Source/ContextualAnimationEditor/Public/ContextualAnimMovieSceneNotifyTrack.h
--- a/Engine/Plugins/Experimental/Animation/ContextualAnimation/Source/ContextualAnimationEditor/Public/ContextualAnimMovieSceneNotifyTrack.h
+++ b/Engine/Plugins/Experimental/Animation/ContextualAnimation/Source/ContextualAnimationEditor/Public/ContextualAnimMovieSceneNotifyTrack.h
@@ -2,10 +2,12 @@
 
 #pragma once
 
-#include "CoreMinimal.h"
 #include "MovieSceneNameableTrack.h"
 #include "ContextualAnimMovieSceneNotifyTrack.generated.h"
 
+enum class EMovieSceneSectionMovedResult;
+struct FMovieSceneSectionMovedParams;
+
 class UAnimSequenceBase;
 struct FAnimNotifyTrack;
 class UContextualAnimMovieSceneNotifySection;
@@ -70,4 +72,8 @@ private:
 	/** Ptr to the animation this track was created from */
 	UPROPERTY()
 	TWeakObjectPtr<UAnimSequenceBase> Animation;
-};
\ No newline at end of file
+};
+
+#if UE_ENABLE_INCLUDE_ORDER_DEPRECATED_IN_5_2
+#include "CoreMinimal.h"
+#endif
diff --git a/Engine/Plugins/Experimental/Animation/ContextualAnimation/Source/ContextualAnimationEditor/Public/ContextualAnimMovieSceneNotifyTrackEditor.h b/Engine/Plugins/Experimental/Animation/ContextualAnimation/Source/ContextualAnimationEditor/Public/ContextualAnimMovieSceneNotifyTrackEditor.h
--- a/Engine/Plugins/Experimental/Animation/ContextualAnimation/Source/ContextualAnimationEditor/Public/ContextualAnimMovieSceneNotifyTrackEditor.h
+++ b/Engine/Plugins/Experimental/Animation/ContextualAnimation/Source/ContextualAnimationEditor/Public/ContextualAnimMovieSceneNotifyTrackEditor.h
@@ -2,10 +2,11 @@
 
 #pragma once
 
-#include "CoreMinimal.h"
+#include "ISequencerSection.h"
 #include "MovieSceneTrackEditor.h"
-#include "ContextualAnimTypes.h"
-#include "BoneContainer.h"
+#include "UObject/GCObject.h"
+
+struct FBuildEditWidgetParams;
 
 class AActor;
 class UAnimMontage;
@@ -79,4 +80,10 @@ public:
 	FString GetNotifyName() const;
 
 	static void PaintNotifyName(FSequencerSectionPainter& Painter, int32 LayerId, const FString& InEventString, float PixelPos, bool bIsEventValid);
-};
\ No newline at end of file
+};
+
+#if UE_ENABLE_INCLUDE_ORDER_DEPRECATED_IN_5_2
+#include "BoneContainer.h"
+#include "ContextualAnimTypes.h"
+#include "CoreMinimal.h"
+#endif
diff --git a/Engine/Plugins/Experimental/Animation/ContextualAnimation/Source/ContextualAnimationEditor/Public/ContextualAnimMovieSceneSequence.h b/Engine/Plugins/Experimental/Animation/ContextualAnimation/Source/ContextualAnimationEditor/Public/ContextualAnimMovieSceneSequence.h
--- a/Engine/Plugins/Experimental/Animation/ContextualAnimation/Source/ContextualAnimationEditor/Public/ContextualAnimMovieSceneSequence.h
+++ b/Engine/Plugins/Experimental/Animation/ContextualAnimation/Source/ContextualAnimationEditor/Public/ContextualAnimMovieSceneSequence.h
@@ -2,8 +2,6 @@
 
 #pragma once
 
-#include "CoreMinimal.h"
-#include "UObject/ObjectMacros.h"
 #include "MovieSceneSequence.h"
 #include "ContextualAnimMovieSceneSequence.generated.h"
 
@@ -43,3 +41,7 @@ private:
 	/** Pointer to the view model that owns this sequence. */
 	TWeakPtr<FContextualAnimViewModel> ViewModelPtr;
 };
+
+#if UE_ENABLE_INCLUDE_ORDER_DEPRECATED_IN_5_2
+#include "CoreMinimal.h"
+#endif
diff --git a/Engine/Plugins/Experimental/Animation/ContextualAnimation/Source/ContextualAnimationEditor/Public/ContextualAnimMovieSceneTrack.h b/Engine/Plugins/Experimental/Animation/ContextualAnimation/Source/ContextualAnimationEditor/Public/ContextualAnimMovieSceneTrack.h
--- a/Engine/Plugins/Experimental/Animation/ContextualAnimation/Source/ContextualAnimationEditor/Public/ContextualAnimMovieSceneTrack.h
+++ b/Engine/Plugins/Experimental/Animation/ContextualAnimation/Source/ContextualAnimationEditor/Public/ContextualAnimMovieSceneTrack.h
@@ -2,10 +2,12 @@
 
 #pragma once
 
-#include "CoreMinimal.h"
 #include "MovieSceneNameableTrack.h"
 #include "ContextualAnimMovieSceneTrack.generated.h"
 
+enum class EMovieSceneSectionMovedResult;
+struct FMovieSceneSectionMovedParams;
+
 class FContextualAnimViewModel;
 
 UCLASS()
@@ -66,3 +68,7 @@ private:
 
 };
 
+
+#if UE_ENABLE_INCLUDE_ORDER_DEPRECATED_IN_5_2
+#include "CoreMinimal.h"
+#endif
diff --git a/Engine/Plugins/Experimental/Animation/ContextualAnimation/Source/ContextualAnimationEditor/Public/ContextualAnimMovieSceneTrackEditor.h b/Engine/Plugins/Experimental/Animation/ContextualAnimation/Source/ContextualAnimationEditor/Public/ContextualAnimMovieSceneTrackEditor.h
--- a/Engine/Plugins/Experimental/Animation/ContextualAnimation/Source/ContextualAnimationEditor/Public/ContextualAnimMovieSceneTrackEditor.h
+++ b/Engine/Plugins/Experimental/Animation/ContextualAnimation/Source/ContextualAnimationEditor/Public/ContextualAnimMovieSceneTrackEditor.h
@@ -2,9 +2,10 @@
 
 #pragma once
 
-#include "CoreMinimal.h"
+#include "ISequencerSection.h"
 #include "MovieSceneTrackEditor.h"
-#include "ContextualAnimTypes.h"
+
+struct FBuildEditWidgetParams;
 
 /** Handles section drawing and manipulation of a ContextualAnimMovieSceneTrack */
 class FContextualAnimMovieSceneTrackEditor : public FMovieSceneTrackEditor
@@ -40,4 +41,9 @@ public:
 	virtual FText GetSectionTitle() const override;
 	virtual float GetSectionHeight() const override;
 	virtual bool SectionIsResizable() const override { return false; }
-};
\ No newline at end of file
+};
+
+#if UE_ENABLE_INCLUDE_ORDER_DEPRECATED_IN_5_2
+#include "ContextualAnimTypes.h"
+#include "CoreMinimal.h"
+#endif
diff --git a/Engine/Plugins/Experimental/Animation/ContextualAnimation/Source/ContextualAnimationEditor/Public/ContextualAnimViewModel.h b/Engine/Plugins/Experimental/Animation/ContextualAnimation/Source/ContextualAnimationEditor/Public/ContextualAnimViewModel.h
--- a/Engine/Plugins/Experimental/Animation/ContextualAnimation/Source/ContextualAnimationEditor/Public/ContextualAnimViewModel.h
+++ b/Engine/Plugins/Experimental/Animation/ContextualAnimation/Source/ContextualAnimationEditor/Public/ContextualAnimViewModel.h
@@ -2,14 +2,16 @@
 
 #pragma once
 
-#include "CoreMinimal.h"
+#include "MovieScene.h"
 #include "UObject/GCObject.h"
-#include "MovieSceneFwd.h"
 #include "ContextualAnimTypes.h"
 #include "ContextualAnimMovieSceneSequence.h"
-#include "ISequencer.h"
 #include "EditorUndoClient.h"
 
+class ISequencer;
+enum class EMovieSceneDataChangeType;
+namespace EMovieScenePlayerStatus { enum Type : int; }
+
 class UWorld;
 class FContextualAnimPreviewScene;
 class UContextualAnimSceneAsset;
@@ -232,4 +234,9 @@ private:
 
 private:
 	bool bInitialized = false;
-};
\ No newline at end of file
+};
+
+#if UE_ENABLE_INCLUDE_ORDER_DEPRECATED_IN_5_2
+#include "CoreMinimal.h"
+#include "ISequencer.h"
+#endif
diff --git a/Engine/Plugins/Experimental/Animation/ContextualAnimation/Source/ContextualAnimationEditor/Public/ContextualAnimationEditor.h b/Engine/Plugins/Experimental/Animation/ContextualAnimation/Source/ContextualAnimationEditor/Public/ContextualAnimationEditor.h
--- a/Engine/Plugins/Experimental/Animation/ContextualAnimation/Source/ContextualAnimationEditor/Public/ContextualAnimationEditor.h
+++ b/Engine/Plugins/Experimental/Animation/ContextualAnimation/Source/ContextualAnimationEditor/Public/ContextualAnimationEditor.h
@@ -2,8 +2,9 @@
 
 #pragma once
 
-#include "CoreMinimal.h"
-#include "Modules/ModuleManager.h"
+#include "Delegates/IDelegateInstance.h"
+#include "Modules/ModuleInterface.h"
+#include "Templates/SharedPointer.h"
 
 class IAssetTypeActions;
 
@@ -23,3 +24,8 @@ private:
 
 	TSharedPtr<IAssetTypeActions> ContextualAnimAssetActions;
 };
+
+#if UE_ENABLE_INCLUDE_ORDER_DEPRECATED_IN_5_2
+#include "CoreMinimal.h"
+#include "Modules/ModuleManager.h"
+#endif
diff --git a/Engine/Plugins/Experimental/Animation/ContextualAnimation/Source/ContextualAnimationEditor/Public/DetailCustomizations/ContextualAnimSceneAssetDetailCustom.h b/Engine/Plugins/Experimental/Animation/ContextualAnimation/Source/ContextualAnimationEditor/Public/DetailCustomizations/ContextualAnimSceneAssetDetailCustom.h
--- a/Engine/Plugins/Experimental/Animation/ContextualAnimation/Source/ContextualAnimationEditor/Public/DetailCustomizations/ContextualAnimSceneAssetDetailCustom.h
+++ b/Engine/Plugins/Experimental/Animation/ContextualAnimation/Source/ContextualAnimationEditor/Public/DetailCustomizations/ContextualAnimSceneAssetDetailCustom.h
@@ -2,8 +2,6 @@
 
 #pragma once
 
-#include "CoreMinimal.h"
-#include "UObject/WeakObjectPtr.h"
 #include "IDetailCustomization.h"
 
 class IDetailLayoutBuilder;
@@ -24,4 +22,9 @@ public:
 protected:
 
 	TWeakPtr<FContextualAnimViewModel> ViewModelPtr;
-};
\ No newline at end of file
+};
+
+#if UE_ENABLE_INCLUDE_ORDER_DEPRECATED_IN_5_2
+#include "CoreMinimal.h"
+#include "UObject/WeakObjectPtr.h"
+#endif
diff --git a/Engine/Plugins/Experimental/Animation/MotionTrajectory/Source/MotionTrajectory/Public/MotionTrajectory.h b/Engine/Plugins/Experimental/Animation/MotionTrajectory/Source/MotionTrajectory/Public/MotionTrajectory.h
--- a/Engine/Plugins/Experimental/Animation/MotionTrajectory/Source/MotionTrajectory/Public/MotionTrajectory.h
+++ b/Engine/Plugins/Experimental/Animation/MotionTrajectory/Source/MotionTrajectory/Public/MotionTrajectory.h
@@ -5,9 +5,8 @@
 #include "Animation/MotionTrajectoryTypes.h"
 #include "Components/ActorComponent.h"
 #include "Containers/RingBuffer.h"
-#include "Modules/ModuleManager.h"
-#include "UObject/ObjectMacros.h"
 
+#include "Modules/ModuleInterface.h"
 #include "MotionTrajectory.generated.h"
 
 DECLARE_LOG_CATEGORY_EXTERN(LogMotionTrajectory, Log, All);
@@ -134,4 +133,8 @@ public:
 	UPROPERTY(EditAnywhere, Category="Debug")
 	bool bDebugDrawTrajectory = false;
 #endif
-};
\ No newline at end of file
+};
+
+#if UE_ENABLE_INCLUDE_ORDER_DEPRECATED_IN_5_2
+#include "Modules/ModuleManager.h"
+#endif
diff --git a/Engine/Plugins/Experimental/Animation/MotionTrajectory/Source/MotionTrajectory/Public/MotionTrajectoryCharacterMovement.h b/Engine/Plugins/Experimental/Animation/MotionTrajectory/Source/MotionTrajectory/Public/MotionTrajectoryCharacterMovement.h
--- a/Engine/Plugins/Experimental/Animation/MotionTrajectory/Source/MotionTrajectory/Public/MotionTrajectoryCharacterMovement.h
+++ b/Engine/Plugins/Experimental/Animation/MotionTrajectory/Source/MotionTrajectory/Public/MotionTrajectoryCharacterMovement.h
@@ -6,6 +6,9 @@
 
 #include "MotionTrajectoryCharacterMovement.generated.h"
 
+struct FTrajectorySample;
+struct FTrajectorySampleRange;
+
 class UCharacterMovementComponent;
 
 // Example motion trajectory component implementation for encapsulating: Character Movement ground locomotion
@@ -91,4 +94,4 @@ protected:
 		const FRotator& PrevRotator,
 		const FRotator& ControlRotationTotalDelta,
 		float IntegrationDelta) const;
-};
\ No newline at end of file
+};
diff --git a/Engine/Plugins/Experimental/Animation/MotionTrajectory/Source/MotionTrajectory/Public/MotionTrajectoryLibrary.h b/Engine/Plugins/Experimental/Animation/MotionTrajectory/Source/MotionTrajectory/Public/MotionTrajectoryLibrary.h
--- a/Engine/Plugins/Experimental/Animation/MotionTrajectory/Source/MotionTrajectory/Public/MotionTrajectoryLibrary.h
+++ b/Engine/Plugins/Experimental/Animation/MotionTrajectory/Source/MotionTrajectory/Public/MotionTrajectoryLibrary.h
@@ -4,7 +4,6 @@
 
 #include "Animation/MotionTrajectoryTypes.h"
 #include "Kismet/BlueprintFunctionLibrary.h"
-#include "UObject/ObjectMacros.h"
 
 #include "MotionTrajectoryLibrary.generated.h"
 
@@ -128,4 +127,4 @@ public:
 			const FTrajectorySampleRange& Trajectory,
 			float Speed = 0.0f,
 			float Tolerance = 0.001f);
-};
\ No newline at end of file
+};
diff --git a/Engine/Plugins/Experimental/ChaosCaching/Source/ChaosCaching/Public/Chaos/Adapters/CacheAdapter.h b/Engine/Plugins/Experimental/ChaosCaching/Source/ChaosCaching/Public/Chaos/Adapters/CacheAdapter.h
--- a/Engine/Plugins/Experimental/ChaosCaching/Source/ChaosCaching/Public/Chaos/Adapters/CacheAdapter.h
+++ b/Engine/Plugins/Experimental/ChaosCaching/Source/ChaosCaching/Public/Chaos/Adapters/CacheAdapter.h
@@ -2,12 +2,12 @@
 
 #pragma once
 
-#include "Chaos/Core.h"
-#include "Chaos/ParticleHandleFwd.h"
+#include "Chaos/Declares.h"
 #include "Features/IModularFeature.h"
-#include "Templates/SubclassOf.h"
-#include "Chaos/PBDRigidsEvolutionFwd.h"
-#include "Chaos/Framework/PhysicsSolverBase.h"
+#include "Math/Transform.h"
+
+namespace Chaos { class FPhysicsSolverEvents; }
+template <class TClass> class TSubclassOf;
 
 class UClass;
 class UChaosCache;
@@ -207,3 +207,9 @@ namespace Chaos
 	};
 
 }    // namespace Chaos
+
+#if UE_ENABLE_INCLUDE_ORDER_DEPRECATED_IN_5_2
+#include "Chaos/Core.h"
+#include "Chaos/Framework/PhysicsSolverBase.h"
+#include "Templates/SubclassOf.h"
+#endif
diff --git a/Engine/Plugins/Experimental/ChaosCaching/Source/ChaosCaching/Public/Chaos/Adapters/GeometryCollectionComponentCacheAdapter.h b/Engine/Plugins/Experimental/ChaosCaching/Source/ChaosCaching/Public/Chaos/Adapters/GeometryCollectionComponentCacheAdapter.h
--- a/Engine/Plugins/Experimental/ChaosCaching/Source/ChaosCaching/Public/Chaos/Adapters/GeometryCollectionComponentCacheAdapter.h
+++ b/Engine/Plugins/Experimental/ChaosCaching/Source/ChaosCaching/Public/Chaos/Adapters/GeometryCollectionComponentCacheAdapter.h
@@ -4,6 +4,7 @@
 
 #include "CacheAdapter.h"
 #include "Chaos/CacheEvents.h"
+#include "Chaos/Framework/PhysicsSolverBase.h"
 #include "EventsData.h"
 #include "GeometryCollectionComponentCacheAdapter.generated.h"
 
diff --git a/Engine/Plugins/Experimental/ChaosCaching/Source/ChaosCaching/Public/Chaos/Adapters/StaticMeshComponentCacheAdapter.h b/Engine/Plugins/Experimental/ChaosCaching/Source/ChaosCaching/Public/Chaos/Adapters/StaticMeshComponentCacheAdapter.h
--- a/Engine/Plugins/Experimental/ChaosCaching/Source/ChaosCaching/Public/Chaos/Adapters/StaticMeshComponentCacheAdapter.h
+++ b/Engine/Plugins/Experimental/ChaosCaching/Source/ChaosCaching/Public/Chaos/Adapters/StaticMeshComponentCacheAdapter.h
@@ -3,6 +3,7 @@
 #pragma once
 
 #include "CacheAdapter.h"
+#include "Chaos/Framework/PhysicsSolverBase.h"
 
 namespace Chaos
 {
diff --git a/Engine/Plugins/Experimental/ChaosCaching/Source/ChaosCaching/Public/Chaos/CacheManagerActor.h b/Engine/Plugins/Experimental/ChaosCaching/Source/ChaosCaching/Public/Chaos/CacheManagerActor.h
--- a/Engine/Plugins/Experimental/ChaosCaching/Source/ChaosCaching/Public/Chaos/CacheManagerActor.h
+++ b/Engine/Plugins/Experimental/ChaosCaching/Source/ChaosCaching/Public/Chaos/CacheManagerActor.h
@@ -5,13 +5,12 @@
 #include "Chaos/ParticleHandleFwd.h"
 #include "Chaos/PBDRigidsEvolutionFwd.h"
 #include "ChaosCache.h"
-#include "Adapters/CacheAdapter.h"
-#include "Engine/EngineTypes.h"
-#include "Chaos/Core.h"
 #include "GameFramework/Actor.h"
 
 #include "CacheManagerActor.generated.h"
 
+namespace Chaos { class FPhysicsSolverEvents; }
+
 class UChaosCacheCollection;
 class UPrimitiveComponent;
 
@@ -338,3 +337,8 @@ public:
 	AChaosCachePlayer(const FObjectInitializer& ObjectInitializer = FObjectInitializer::Get());
 
 };
+
+#if UE_ENABLE_INCLUDE_ORDER_DEPRECATED_IN_5_2
+#include "Adapters/CacheAdapter.h"
+#include "Chaos/Core.h"
+#endif
diff --git a/Engine/Plugins/Experimental/ChaosCaching/Source/ChaosCaching/Public/Chaos/ChaosCache.h b/Engine/Plugins/Experimental/ChaosCaching/Source/ChaosCaching/Public/Chaos/ChaosCache.h
--- a/Engine/Plugins/Experimental/ChaosCaching/Source/ChaosCaching/Public/Chaos/ChaosCache.h
+++ b/Engine/Plugins/Experimental/ChaosCaching/Source/ChaosCaching/Public/Chaos/ChaosCache.h
@@ -6,9 +6,7 @@
 #include "CacheEvents.h"
 #include "Containers/Queue.h"
 #include "Curves/RichCurve.h"
-#include "Logging/LogMacros.h"
 
-#include <atomic>
 
 #include "ChaosCache.generated.h"
 
diff --git a/Engine/Plugins/Experimental/ChaosCaching/Source/ChaosCaching/Public/Chaos/ChaosCachingPlugin.h b/Engine/Plugins/Experimental/ChaosCaching/Source/ChaosCaching/Public/Chaos/ChaosCachingPlugin.h
--- a/Engine/Plugins/Experimental/ChaosCaching/Source/ChaosCaching/Public/Chaos/ChaosCachingPlugin.h
+++ b/Engine/Plugins/Experimental/ChaosCaching/Source/ChaosCaching/Public/Chaos/ChaosCachingPlugin.h
@@ -2,10 +2,7 @@
 
 #pragma once
 
-#include "CoreMinimal.h"
-#include "Modules/ModuleInterface.h"
 #include "Modules/ModuleManager.h"
-#include "Logging/LogMacros.h"
 
 namespace Chaos {
 	class FGeometryCollectionCacheAdapter;
@@ -48,3 +45,7 @@ private:
 
 	FDelegateHandle OnCreateMovieSceneObjectSpawnerHandle;
 };
+
+#if UE_ENABLE_INCLUDE_ORDER_DEPRECATED_IN_5_2
+#include "CoreMinimal.h"
+#endif
diff --git a/Engine/Plugins/Experimental/ChaosCaching/Source/ChaosCaching/Public/Chaos/Sequencer/MovieSceneChaosCacheTrack.h b/Engine/Plugins/Experimental/ChaosCaching/Source/ChaosCaching/Public/Chaos/Sequencer/MovieSceneChaosCacheTrack.h
--- a/Engine/Plugins/Experimental/ChaosCaching/Source/ChaosCaching/Public/Chaos/Sequencer/MovieSceneChaosCacheTrack.h
+++ b/Engine/Plugins/Experimental/ChaosCaching/Source/ChaosCaching/Public/Chaos/Sequencer/MovieSceneChaosCacheTrack.h
@@ -3,7 +3,6 @@
 #pragma once
 
 #include "MovieSceneNameableTrack.h"
-#include "Chaos/CacheManagerActor.h"
 #include "Compilation/IMovieSceneTrackTemplateProducer.h"
 #include "MovieSceneChaosCacheTrack.generated.h"
 
@@ -51,3 +50,7 @@ private:
 	UPROPERTY()
 	TArray<TObjectPtr<UMovieSceneSection>> AnimationSections;
 };
+
+#if UE_ENABLE_INCLUDE_ORDER_DEPRECATED_IN_5_2
+#include "Chaos/CacheManagerActor.h"
+#endif
diff --git a/Engine/Plugins/Experimental/ChaosCaching/Source/ChaosCachingEditor/Public/Chaos/CacheCollectionCustomization.h b/Engine/Plugins/Experimental/ChaosCaching/Source/ChaosCachingEditor/Public/Chaos/CacheCollectionCustomization.h
--- a/Engine/Plugins/Experimental/ChaosCaching/Source/ChaosCachingEditor/Public/Chaos/CacheCollectionCustomization.h
+++ b/Engine/Plugins/Experimental/ChaosCaching/Source/ChaosCachingEditor/Public/Chaos/CacheCollectionCustomization.h
@@ -3,9 +3,9 @@
 #pragma once
 
 #include "IDetailCustomization.h"
-#include "Templates/SharedPointer.h"
-#include "UObject/WeakObjectPtrTemplates.h"
-#include "Types/SlateEnums.h"
+#include "UObject/WeakObjectPtr.h"
+
+namespace ETextCommit { enum Type : int; }
 
 class IDetailLayoutBuilder;
 class IPropertyHandle;
@@ -38,4 +38,8 @@ private:
 	TWeakObjectPtr<UObject> Item;
 
 	TMap< int, TSharedPtr<SEditableTextBox>> NameEditBoxes;
-};
\ No newline at end of file
+};
+
+#if UE_ENABLE_INCLUDE_ORDER_DEPRECATED_IN_5_2
+#include "Types/SlateEnums.h"
+#endif
diff --git a/Engine/Plugins/Experimental/ChaosCaching/Source/ChaosCachingEditor/Public/Chaos/CacheManagerCustomization.h b/Engine/Plugins/Experimental/ChaosCaching/Source/ChaosCachingEditor/Public/Chaos/CacheManagerCustomization.h
--- a/Engine/Plugins/Experimental/ChaosCaching/Source/ChaosCachingEditor/Public/Chaos/CacheManagerCustomization.h
+++ b/Engine/Plugins/Experimental/ChaosCaching/Source/ChaosCachingEditor/Public/Chaos/CacheManagerCustomization.h
@@ -3,9 +3,10 @@
 #pragma once
 
 #include "IDetailCustomization.h"
-#include "Templates/SharedPointer.h"
 #include "IPropertyTypeCustomization.h"
-#include "Widgets/Input/SComboBox.h"
+
+class SWidget;
+namespace ESelectInfo { enum Type : int; }
 
 class IDetailLayoutBuilder;
 class IPropertyHandle;
@@ -39,4 +40,8 @@ public:
 	static TSharedRef<IPropertyTypeCustomization> MakeInstance();
 	virtual void CustomizeHeader(TSharedRef<IPropertyHandle> PropertyHandle, FDetailWidgetRow& HeaderRow, IPropertyTypeCustomizationUtils& CustomizationUtils) override;
 	virtual void CustomizeChildren(TSharedRef<IPropertyHandle> PropertyHandle, IDetailChildrenBuilder& ChildBuilder, IPropertyTypeCustomizationUtils& CustomizationUtils) override;
-};
\ No newline at end of file
+};
+
+#if UE_ENABLE_INCLUDE_ORDER_DEPRECATED_IN_5_2
+#include "Widgets/Input/SComboBox.h"
+#endif
diff --git a/Engine/Plugins/Experimental/ChaosCaching/Source/ChaosCachingEditor/Public/Chaos/ChaosCachingEditorPlugin.h b/Engine/Plugins/Experimental/ChaosCaching/Source/ChaosCachingEditor/Public/Chaos/ChaosCachingEditorPlugin.h
--- a/Engine/Plugins/Experimental/ChaosCaching/Source/ChaosCachingEditor/Public/Chaos/ChaosCachingEditorPlugin.h
+++ b/Engine/Plugins/Experimental/ChaosCaching/Source/ChaosCachingEditor/Public/Chaos/ChaosCachingEditorPlugin.h
@@ -2,13 +2,14 @@
 
 #pragma once
 
-#include "CoreMinimal.h"
-#include "Modules/ModuleInterface.h"
 #include "Modules/ModuleManager.h"
-#include "Delegates/IDelegateInstance.h"
-#include "Framework/MultiBox/MultiBoxExtender.h"
 #include "Sequencer/MovieSceneChaosCacheTrackRecorder.h"
 
+class FExtender;
+class FMenuBuilder;
+class FUICommandList;
+class IConsoleObject;
+
 class FAssetTypeActions_ChaosCacheCollection;
 
 class UToolMenu;
@@ -58,3 +59,8 @@ private:
 
 	FMovieSceneChaosCacheTrackRecorderFactory MovieSceneChaosCacheTrackRecorder;
 };
+
+#if UE_ENABLE_INCLUDE_ORDER_DEPRECATED_IN_5_2
+#include "CoreMinimal.h"
+#include "Framework/MultiBox/MultiBoxExtender.h"
+#endif
diff --git a/Engine/Plugins/Experimental/ChaosCaching/Source/ChaosCachingEditor/Public/Sequencer/ChaosCacheTrackEditor.h b/Engine/Plugins/Experimental/ChaosCaching/Source/ChaosCachingEditor/Public/Sequencer/ChaosCacheTrackEditor.h
--- a/Engine/Plugins/Experimental/ChaosCaching/Source/ChaosCachingEditor/Public/Sequencer/ChaosCacheTrackEditor.h
+++ b/Engine/Plugins/Experimental/ChaosCaching/Source/ChaosCachingEditor/Public/Sequencer/ChaosCacheTrackEditor.h
@@ -1,12 +1,12 @@
-﻿// Copyright Epic Games, Inc. All Rights Reserved.
+// Copyright Epic Games, Inc. All Rights Reserved.
 
 #pragma once
 
-#include "ISequencer.h"
 #include "ISequencerSection.h"
-#include "ISequencerTrackEditor.h"
 #include "MovieSceneTrackEditor.h"
 
+struct FBuildEditWidgetParams;
+
 /**
  * Editor for chaos Cache tracks
  */
diff --git a/Engine/Plugins/Experimental/ChaosCaching/Source/ChaosCachingEditor/Public/Sequencer/MovieSceneChaosCacheTrackRecorder.h b/Engine/Plugins/Experimental/ChaosCaching/Source/ChaosCachingEditor/Public/Sequencer/MovieSceneChaosCacheTrackRecorder.h
--- a/Engine/Plugins/Experimental/ChaosCaching/Source/ChaosCachingEditor/Public/Sequencer/MovieSceneChaosCacheTrackRecorder.h
+++ b/Engine/Plugins/Experimental/ChaosCaching/Source/ChaosCachingEditor/Public/Sequencer/MovieSceneChaosCacheTrackRecorder.h
@@ -1,16 +1,15 @@
-﻿// Copyright Epic Games, Inc. All Rights Reserved.
+// Copyright Epic Games, Inc. All Rights Reserved.
 
 #pragma once
 
 #include "TrackRecorders/MovieSceneTrackRecorder.h"
-#include "CoreMinimal.h"
-#include "MovieScene.h"
 #include "TrackRecorders/IMovieSceneTrackRecorderFactory.h"
-#include "Channels/MovieSceneFloatChannel.h"
-#include "Chaos/CacheManagerActor.h"
 
 #include "MovieSceneChaosCacheTrackRecorder.generated.h"
 
+class AChaosCacheManager;
+enum class ECacheMode : uint8;
+
 class  UMovieSceneChaosCacheTrack;
 class  UMovieSceneChaosCacheSection;
 
@@ -76,3 +75,9 @@ private:
 	FFrameNumber RecordStartFrame;
 };
 
+
+#if UE_ENABLE_INCLUDE_ORDER_DEPRECATED_IN_5_2
+#include "Channels/MovieSceneFloatChannel.h"
+#include "Chaos/CacheManagerActor.h"
+#include "CoreMinimal.h"
+#endif
diff --git a/Engine/Plugins/Experimental/ChaosCaching/Source/ChaosCachingEditor/Public/Sequencer/TakeRecorderChaosCacheSource.h b/Engine/Plugins/Experimental/ChaosCaching/Source/ChaosCachingEditor/Public/Sequencer/TakeRecorderChaosCacheSource.h
--- a/Engine/Plugins/Experimental/ChaosCaching/Source/ChaosCachingEditor/Public/Sequencer/TakeRecorderChaosCacheSource.h
+++ b/Engine/Plugins/Experimental/ChaosCaching/Source/ChaosCachingEditor/Public/Sequencer/TakeRecorderChaosCacheSource.h
@@ -2,12 +2,12 @@
 
 #pragma once
 
-#include "CoreMinimal.h"
 #include "TakeRecorderSource.h"
-#include "UObject/WeakObjectPtrTemplates.h"
 
 #include "TakeRecorderChaosCacheSource.generated.h"
 
+struct FPropertyChangedEvent;
+
 class ULevelSequence;
 class UMovieSceneFolder;
 class UMovieSceneChaosCacheTrack;
@@ -60,3 +60,7 @@ private:
 	 */
 	TWeakObjectPtr<UMovieSceneChaosCacheTrack> CachedChaosCacheTrack;
 };
+
+#if UE_ENABLE_INCLUDE_ORDER_DEPRECATED_IN_5_2
+#include "CoreMinimal.h"
+#endif
diff --git a/Engine/Plugins/Experimental/ChaosUserDataPT/Source/ChaosUserDataPT/Public/ChaosUserDataPTStats.h b/Engine/Plugins/Experimental/ChaosUserDataPT/Source/ChaosUserDataPT/Public/ChaosUserDataPTStats.h
--- a/Engine/Plugins/Experimental/ChaosUserDataPT/Source/ChaosUserDataPT/Public/ChaosUserDataPTStats.h
+++ b/Engine/Plugins/Experimental/ChaosUserDataPT/Source/ChaosUserDataPT/Public/ChaosUserDataPTStats.h
@@ -4,6 +4,8 @@
 
 #include "Stats/Stats2.h"
 
+enum class EStatFlags : uint8;
+
 DECLARE_STATS_GROUP(TEXT("ChaosUserDataPT"), STATGROUP_ChaosUserDataPT, STATCAT_Advanced);
 
 namespace Chaos
diff --git a/Engine/Plugins/Experimental/FullBodyIK/Source/FullBodyIK/Public/FBIKConstraint.h b/Engine/Plugins/Experimental/FullBodyIK/Source/FullBodyIK/Public/FBIKConstraint.h
--- a/Engine/Plugins/Experimental/FullBodyIK/Source/FullBodyIK/Public/FBIKConstraint.h
+++ b/Engine/Plugins/Experimental/FullBodyIK/Source/FullBodyIK/Public/FBIKConstraint.h
@@ -2,8 +2,7 @@
 
 #pragma once
 
-#include "CoreMinimal.h"
-#include "UObject/ObjectMacros.h"
+#include "Math/Transform.h"
 #include "Misc/TVariant.h"
 
 enum class EConstraintType : uint8
@@ -349,4 +348,9 @@ struct FULLBODYIK_API FPoleVectorConstraint
 };
 
 
-using ConstraintType = TVariant<FRotationLimitConstraint, FPositionLimitConstraint, FPoleVectorConstraint>;
\ No newline at end of file
+using ConstraintType = TVariant<FRotationLimitConstraint, FPositionLimitConstraint, FPoleVectorConstraint>;
+
+#if UE_ENABLE_INCLUDE_ORDER_DEPRECATED_IN_5_2
+#include "CoreMinimal.h"
+#include "UObject/ObjectMacros.h"
+#endif
diff --git a/Engine/Plugins/Experimental/FullBodyIK/Source/FullBodyIK/Public/FBIKShared.h b/Engine/Plugins/Experimental/FullBodyIK/Source/FullBodyIK/Public/FBIKShared.h
--- a/Engine/Plugins/Experimental/FullBodyIK/Source/FullBodyIK/Public/FBIKShared.h
+++ b/Engine/Plugins/Experimental/FullBodyIK/Source/FullBodyIK/Public/FBIKShared.h
@@ -2,11 +2,11 @@
 
 #pragma once
 
-#include "CoreMinimal.h"
 #include "JacobianSolver.h"
-#include "Rigs/RigHierarchyDefines.h"
 #include "FBIKShared.generated.h"
 
+namespace JacobianIK { struct FFBIKEffectorTarget; }
+
 class FULLBODYIK_API FJacobianSolver_FullbodyIK : public FJacobianSolverBase
 {
 public:
@@ -74,3 +74,8 @@ struct FSolverInput
 	bool bUseJacobianTranspose = false;
 };
 
+
+#if UE_ENABLE_INCLUDE_ORDER_DEPRECATED_IN_5_2
+#include "CoreMinimal.h"
+#include "Rigs/RigHierarchyDefines.h"
+#endif
diff --git a/Engine/Plugins/Experimental/FullBodyIK/Source/FullBodyIK/Public/FBIKUtil.h b/Engine/Plugins/Experimental/FullBodyIK/Source/FullBodyIK/Public/FBIKUtil.h
--- a/Engine/Plugins/Experimental/FullBodyIK/Source/FullBodyIK/Public/FBIKUtil.h
+++ b/Engine/Plugins/Experimental/FullBodyIK/Source/FullBodyIK/Public/FBIKUtil.h
@@ -2,7 +2,7 @@
 
 #pragma once
 
-#include "CoreMinimal.h"
+#include "Math/Quat.h"
 
 class URigHierarchy;
 struct FRigElementKey;
@@ -16,3 +16,7 @@ namespace FBIKUtil
 };
 
 
+
+#if UE_ENABLE_INCLUDE_ORDER_DEPRECATED_IN_5_2
+#include "CoreMinimal.h"
+#endif
diff --git a/Engine/Plugins/Experimental/FullBodyIK/Source/FullBodyIK/Public/FullBodyIK.h b/Engine/Plugins/Experimental/FullBodyIK/Source/FullBodyIK/Public/FullBodyIK.h
--- a/Engine/Plugins/Experimental/FullBodyIK/Source/FullBodyIK/Public/FullBodyIK.h
+++ b/Engine/Plugins/Experimental/FullBodyIK/Source/FullBodyIK/Public/FullBodyIK.h
@@ -2,8 +2,8 @@
 
 #pragma once
 
-#include "CoreMinimal.h"
-#include "Modules/ModuleManager.h"
+#include "Logging/LogMacros.h"
+#include "Modules/ModuleInterface.h"
 
 class FFullBodyIKModule : public IModuleInterface
 {
@@ -15,3 +15,8 @@ public:
 };
 
 DECLARE_LOG_CATEGORY_EXTERN(LogIKSolver, Warning, All);
+
+#if UE_ENABLE_INCLUDE_ORDER_DEPRECATED_IN_5_2
+#include "CoreMinimal.h"
+#include "Modules/ModuleManager.h"
+#endif
diff --git a/Engine/Plugins/Experimental/FullBodyIK/Source/FullBodyIK/Public/JacobianIK.h b/Engine/Plugins/Experimental/FullBodyIK/Source/FullBodyIK/Public/JacobianIK.h
--- a/Engine/Plugins/Experimental/FullBodyIK/Source/FullBodyIK/Public/JacobianIK.h
+++ b/Engine/Plugins/Experimental/FullBodyIK/Source/FullBodyIK/Public/JacobianIK.h
@@ -2,9 +2,7 @@
 
 #pragma once
 
-#include "CoreMinimal.h"
-#include "UObject/ObjectMacros.h"
-#include "FBIKConstraint.h"
+#include "Math/Quat.h"
 
 // Just to be sure, also added this in Eigen.Build.cs
 #ifndef EIGEN_MPL2_ONLY
@@ -24,6 +22,8 @@ PRAGMA_DEFAULT_VISIBILITY_END
 #pragma warning(pop)
 #endif
 
+struct FFBIKLinkData;
+
 namespace JacobianIK
 {
 	struct FFBIKEffectorTarget
@@ -183,3 +183,9 @@ namespace JacobianIK
 		const FSolverParameter& SolverParam);
 };
 
+
+#if UE_ENABLE_INCLUDE_ORDER_DEPRECATED_IN_5_2
+#include "CoreMinimal.h"
+#include "FBIKConstraint.h"
+#include "UObject/ObjectMacros.h"
+#endif
diff --git a/Engine/Plugins/Experimental/FullBodyIK/Source/FullBodyIK/Public/JacobianSolver.h b/Engine/Plugins/Experimental/FullBodyIK/Source/FullBodyIK/Public/JacobianSolver.h
--- a/Engine/Plugins/Experimental/FullBodyIK/Source/FullBodyIK/Public/JacobianSolver.h
+++ b/Engine/Plugins/Experimental/FullBodyIK/Source/FullBodyIK/Public/JacobianSolver.h
@@ -3,6 +3,7 @@
 #pragma once
 
 #include "JacobianIK.h"
+#include "Math/Transform.h"
 
 using namespace JacobianIK;
 
diff --git a/Engine/Plugins/Experimental/FullBodyIK/Source/PBIK/Public/Core/PBIKBody.h b/Engine/Plugins/Experimental/FullBodyIK/Source/PBIK/Public/Core/PBIKBody.h
--- a/Engine/Plugins/Experimental/FullBodyIK/Source/PBIK/Public/Core/PBIKBody.h
+++ b/Engine/Plugins/Experimental/FullBodyIK/Source/PBIK/Public/Core/PBIKBody.h
@@ -2,8 +2,6 @@
 
 #pragma once
 
-#include "UObject/NameTypes.h"
-#include "Math/Vector.h"
 #include "Math/Quat.h"
 
 struct FPBIKSolverSettings;
@@ -125,4 +123,4 @@ inline bool operator<(const FRigidBody& Lhs, const FRigidBody& Rhs)
 	return Lhs.GetNumBonesToRoot() < Rhs.GetNumBonesToRoot(); 
 }
 
-} // namespace
\ No newline at end of file
+} // namespace
diff --git a/Engine/Plugins/Experimental/FullBodyIK/Source/PBIK/Public/Core/PBIKConstraint.h b/Engine/Plugins/Experimental/FullBodyIK/Source/PBIK/Public/Core/PBIKConstraint.h
--- a/Engine/Plugins/Experimental/FullBodyIK/Source/PBIK/Public/Core/PBIKConstraint.h
+++ b/Engine/Plugins/Experimental/FullBodyIK/Source/PBIK/Public/Core/PBIKConstraint.h
@@ -2,7 +2,9 @@
 
 #pragma once
 
-#include "PBIKBody.h"
+#include "Math/Quat.h"
+
+namespace PBIK { struct FRigidBody; }
 
 namespace PBIK
 {
@@ -122,3 +124,7 @@ private:
 
 } // namespace
 
+
+#if UE_ENABLE_INCLUDE_ORDER_DEPRECATED_IN_5_2
+#include "PBIKBody.h"
+#endif
diff --git a/Engine/Plugins/Experimental/FullBodyIK/Source/PBIK/Public/Core/PBIKSolver.h b/Engine/Plugins/Experimental/FullBodyIK/Source/PBIK/Public/Core/PBIKSolver.h
--- a/Engine/Plugins/Experimental/FullBodyIK/Source/PBIK/Public/Core/PBIKSolver.h
+++ b/Engine/Plugins/Experimental/FullBodyIK/Source/PBIK/Public/Core/PBIKSolver.h
@@ -2,12 +2,15 @@
 
 #pragma once
 
-#include "PBIKBody.h"
-#include "PBIKConstraint.h"
 #include "PBIKDebug.h"
 
 #include "PBIKSolver.generated.h"
 
+namespace PBIK { struct FBone; }
+namespace PBIK { struct FBoneSettings; }
+namespace PBIK { struct FConstraint; }
+namespace PBIK { struct FPinConstraint; }
+
 DECLARE_CYCLE_STAT(TEXT("PBIK Solve"), STAT_PBIK_Solve, STATGROUP_Anim);
 
 namespace PBIK
@@ -282,4 +285,9 @@ private:
 	
 	PBIK::FDebugDraw DebugDraw;
 	friend PBIK::FDebugDraw;
-};
\ No newline at end of file
+};
+
+#if UE_ENABLE_INCLUDE_ORDER_DEPRECATED_IN_5_2
+#include "PBIKBody.h"
+#include "PBIKConstraint.h"
+#endif
diff --git a/Engine/Plugins/Experimental/FullBodyIK/Source/PBIK/Public/PBIK.h b/Engine/Plugins/Experimental/FullBodyIK/Source/PBIK/Public/PBIK.h
--- a/Engine/Plugins/Experimental/FullBodyIK/Source/PBIK/Public/PBIK.h
+++ b/Engine/Plugins/Experimental/FullBodyIK/Source/PBIK/Public/PBIK.h
@@ -2,8 +2,8 @@
 
 #pragma once
 
-#include "CoreMinimal.h"
-#include "Modules/ModuleManager.h"
+#include "Logging/LogMacros.h"
+#include "Modules/ModuleInterface.h"
 
 class FPBIKModule : public IModuleInterface
 {
@@ -16,3 +16,8 @@ public:
 
 DECLARE_LOG_CATEGORY_EXTERN(LogPBIKSolver, Warning, All);
 
+
+#if UE_ENABLE_INCLUDE_ORDER_DEPRECATED_IN_5_2
+#include "CoreMinimal.h"
+#include "Modules/ModuleManager.h"
+#endif
diff --git a/Engine/Plugins/Experimental/FullBodyIK/Source/PBIK/Public/PBIK_Shared.h b/Engine/Plugins/Experimental/FullBodyIK/Source/PBIK/Public/PBIK_Shared.h
--- a/Engine/Plugins/Experimental/FullBodyIK/Source/PBIK/Public/PBIK_Shared.h
+++ b/Engine/Plugins/Experimental/FullBodyIK/Source/PBIK/Public/PBIK_Shared.h
@@ -2,10 +2,9 @@
 
 #pragma once
 
-#include "CoreMinimal.h"
 
-#include "Core/PBIKSolver.h"
 
+#include "Core/PBIKBody.h"
 #include "PBIK_Shared.generated.h"
 
 UENUM(BlueprintType)
@@ -108,4 +107,9 @@ struct PBIK_API FPBIKBoneSetting
 		Settings.PreferredAngles.Yaw = PreferredAngles.Z;
 		Settings.PreferredAngles.Roll = PreferredAngles.X;
 	}
-};
\ No newline at end of file
+};
+
+#if UE_ENABLE_INCLUDE_ORDER_DEPRECATED_IN_5_2
+#include "Core/PBIKSolver.h"
+#include "CoreMinimal.h"
+#endif
diff --git a/Engine/Plugins/Experimental/FullBodyIK/Source/PBIK/Public/RigUnit_PBIK.h b/Engine/Plugins/Experimental/FullBodyIK/Source/PBIK/Public/RigUnit_PBIK.h
--- a/Engine/Plugins/Experimental/FullBodyIK/Source/PBIK/Public/RigUnit_PBIK.h
+++ b/Engine/Plugins/Experimental/FullBodyIK/Source/PBIK/Public/RigUnit_PBIK.h
@@ -2,6 +2,7 @@
 
 #pragma once
 
+#include "RigVMCore/RigVMFunction.h"
 #include "Units/Highlevel/RigUnit_HighlevelBase.h"
 
 #include "Core/PBIKSolver.h"
@@ -9,8 +10,11 @@
 
 #include "PBIK_Shared.h"
 
+#include "RigVMCore/RigVMMemoryStorage.h"
 #include "RigUnit_PBIK.generated.h"
 
+struct FControlRigExecuteContext;
+
 using PBIK::FDebugLine;
 
 USTRUCT(BlueprintType)
diff --git a/Engine/Plugins/Experimental/GeometryCollectionPlugin/Source/GeometryCollectionNodes/Private/Dataflow/GeometryCollectionNodes.cpp b/Engine/Plugins/Experimental/GeometryCollectionPlugin/Source/GeometryCollectionNodes/Private/Dataflow/GeometryCollectionNodes.cpp
--- a/Engine/Plugins/Experimental/GeometryCollectionPlugin/Source/GeometryCollectionNodes/Private/Dataflow/GeometryCollectionNodes.cpp
+++ b/Engine/Plugins/Experimental/GeometryCollectionPlugin/Source/GeometryCollectionNodes/Private/Dataflow/GeometryCollectionNodes.cpp
@@ -5,6 +5,7 @@
 
 #include "Engine/Engine.h"
 #include "Engine/StaticMesh.h"
+#include "GeometryCollection/Facades/CollectionMeshFacade.h"
 #include "GeometryCollection/GeometryCollectionObject.h"
 #include "GeometryCollection/ManagedArrayCollection.h"
 #include "GeometryCollection/GeometryCollection.h"
diff --git a/Engine/Plugins/Experimental/Mutable/Source/CustomizableObject/Classes/MuCO/CustomizableObject.h b/Engine/Plugins/Experimental/Mutable/Source/CustomizableObject/Classes/MuCO/CustomizableObject.h
--- a/Engine/Plugins/Experimental/Mutable/Source/CustomizableObject/Classes/MuCO/CustomizableObject.h
+++ b/Engine/Plugins/Experimental/Mutable/Source/CustomizableObject/Classes/MuCO/CustomizableObject.h
@@ -6,7 +6,6 @@
 #include "MuCO/CustomizableObjectIdentifier.h"
 #include "MuCO/CustomizableObjectUIData.h"
 #include "RHIDefinitions.h"
-#include "Delegates/IDelegateInstance.h"
 
 class FReply;
 
diff --git a/Engine/Plugins/Experimental/Mutable/Source/CustomizableObject/Private/MuCO/CustomizableObject.cpp b/Engine/Plugins/Experimental/Mutable/Source/CustomizableObject/Private/MuCO/CustomizableObject.cpp
--- a/Engine/Plugins/Experimental/Mutable/Source/CustomizableObject/Private/MuCO/CustomizableObject.cpp
+++ b/Engine/Plugins/Experimental/Mutable/Source/CustomizableObject/Private/MuCO/CustomizableObject.cpp
@@ -5,7 +5,6 @@
 #include "Algo/Copy.h"
 #include "AssetRegistry/AssetRegistryModule.h"
 #include "Async/AsyncFileHandle.h"
-#include "Delegates/IDelegateInstance.h"
 #include "EdGraph/EdGraph.h"
 #include "Engine/SkeletalMesh.h"
 #include "HAL/FileManager.h"
@@ -26,7 +25,6 @@
 #include "Serialization/MemoryReader.h"
 #include "Serialization/MemoryWriter.h"
 #include "UObject/ObjectSaveContext.h"
-#include "Templates/UniquePtr.h"
 
 #if WITH_EDITOR
 #include "Editor.h"
diff --git a/Engine/Plugins/Experimental/Mutable/Source/CustomizableObjectEditor/Classes/MuCOE/Nodes/CustomizableObjectNodeTextureSaturate.h b/Engine/Plugins/Experimental/Mutable/Source/CustomizableObjectEditor/Classes/MuCOE/Nodes/CustomizableObjectNodeTextureSaturate.h
--- a/Engine/Plugins/Experimental/Mutable/Source/CustomizableObjectEditor/Classes/MuCOE/Nodes/CustomizableObjectNodeTextureSaturate.h
+++ b/Engine/Plugins/Experimental/Mutable/Source/CustomizableObjectEditor/Classes/MuCOE/Nodes/CustomizableObjectNodeTextureSaturate.h
@@ -1,16 +1,12 @@
 // Copyright Epic Games, Inc. All Rights Reserved.
 #pragma once
 
-#include "EdGraph/EdGraphNode.h"
-#include "EdGraph/EdGraphPin.h"
-#include "HAL/Platform.h"
-#include "Internationalization/Text.h"
-#include "Math/Color.h"
 #include "MuCOE/Nodes/CustomizableObjectNode.h"
-#include "UObject/UObjectGlobals.h"
 
 #include "CustomizableObjectNodeTextureSaturate.generated.h"
 
+namespace ENodeTitleType { enum Type : int; }
+
 class FArchive;
 class UCustomizableObjectNodeRemapPins;
 class UObject;
diff --git a/Engine/Plugins/Experimental/Mutable/Source/CustomizableObjectEditor/Private/MuCOE/SMutableCodeViewer.cpp b/Engine/Plugins/Experimental/Mutable/Source/CustomizableObjectEditor/Private/MuCOE/SMutableCodeViewer.cpp
--- a/Engine/Plugins/Experimental/Mutable/Source/CustomizableObjectEditor/Private/MuCOE/SMutableCodeViewer.cpp
+++ b/Engine/Plugins/Experimental/Mutable/Source/CustomizableObjectEditor/Private/MuCOE/SMutableCodeViewer.cpp
@@ -28,6 +28,7 @@
 #include "MuT/Streams.h"
 #include "MuT/Table.h"
 #include "MuT/TypeInfo.h"
+#include "Widgets/Colors/SColorBlock.h"
 #include "Widgets/Input/SButton.h"
 #include "Widgets/Input/SNumericEntryBox.h"
 #include "Widgets/Input/SComboBox.h"
diff --git a/Engine/Plugins/Experimental/PCG/Source/PCG/Private/Elements/PCGGather.cpp b/Engine/Plugins/Experimental/PCG/Source/PCG/Private/Elements/PCGGather.cpp
--- a/Engine/Plugins/Experimental/PCG/Source/PCG/Private/Elements/PCGGather.cpp
+++ b/Engine/Plugins/Experimental/PCG/Source/PCG/Private/Elements/PCGGather.cpp
@@ -1,7 +1,6 @@
 // Copyright Epic Games, Inc. All Rights Reserved.
 
 #include "Elements/PCGGather.h"
-#include "Helpers/PCGSettingsHelpers.h"
 #include "PCGContext.h"
 #include "PCGPin.h"
 	
diff --git a/Engine/Plugins/Experimental/PCG/Source/PCG/Private/Tests/Elements/PCGGatherTest.cpp b/Engine/Plugins/Experimental/PCG/Source/PCG/Private/Tests/Elements/PCGGatherTest.cpp
--- a/Engine/Plugins/Experimental/PCG/Source/PCG/Private/Tests/Elements/PCGGatherTest.cpp
+++ b/Engine/Plugins/Experimental/PCG/Source/PCG/Private/Tests/Elements/PCGGatherTest.cpp
@@ -1,13 +1,9 @@
 // Copyright Epic Games, Inc. All Rights Reserved.
 
 #include "Tests/Determinism/PCGDeterminismTestsCommon.h"
-#include "Tests/PCGTestsCommon.h"
 
 #include "PCGComponent.h"
 #include "PCGContext.h"
-#include "PCGData.h"
-#include "PCGHelpers.h"
-#include "PCGParamData.h"
 #include "Data/PCGPointData.h"
 
 #include "Elements/PCGGather.h"
diff --git a/Engine/Plugins/Experimental/PCG/Source/PCG/Public/Elements/PCGGather.h b/Engine/Plugins/Experimental/PCG/Source/PCG/Public/Elements/PCGGather.h
--- a/Engine/Plugins/Experimental/PCG/Source/PCG/Public/Elements/PCGGather.h
+++ b/Engine/Plugins/Experimental/PCG/Source/PCG/Public/Elements/PCGGather.h
@@ -1,9 +1,7 @@
 // Copyright Epic Games, Inc. All Rights Reserved.
 #pragma once
 
-#include "CoreMinimal.h"
 #include "PCGSettings.h"
-#include "PCGElement.h"
 
 #include "PCGGather.generated.h"
 
@@ -36,3 +34,7 @@ protected:
 	virtual bool ExecuteInternal(FPCGContext* Context) const override;
 };
 
+
+#if UE_ENABLE_INCLUDE_ORDER_DEPRECATED_IN_5_2
+#include "CoreMinimal.h"
+#endif
diff --git a/Engine/Plugins/Experimental/PCG/Source/PCG/Public/Elements/PCGPointMatchAndSet.h b/Engine/Plugins/Experimental/PCG/Source/PCG/Public/Elements/PCGPointMatchAndSet.h
--- a/Engine/Plugins/Experimental/PCG/Source/PCG/Public/Elements/PCGPointMatchAndSet.h
+++ b/Engine/Plugins/Experimental/PCG/Source/PCG/Public/Elements/PCGPointMatchAndSet.h
@@ -2,12 +2,14 @@
 
 #pragma once
 
+#include "Metadata/PCGMetadataTypesConstantStruct.h"
 #include "PCGSettings.h"
-#include "MatchAndSet/PCGMatchAndSetBase.h"
 #include "Metadata/PCGAttributePropertySelector.h"
 
 #include "PCGPointMatchAndSet.generated.h"
 
+class UPCGMatchAndSetBase;
+
 /** This settings class is used to create a PCG node that will apply a "Match and Set" operation
 * on the point data it consumes as input.
 * E.g. for a given point, if it matches with something in the Match & Set object, it will set a value on the point.
@@ -82,4 +84,8 @@ class FPCGPointMatchAndSetElement : public FSimplePCGElement
 {
 public:
 	virtual bool ExecuteInternal(FPCGContext* Context) const override;
-};
\ No newline at end of file
+};
+
+#if UE_ENABLE_INCLUDE_ORDER_DEPRECATED_IN_5_2
+#include "MatchAndSet/PCGMatchAndSetBase.h"
+#endif
diff --git a/Engine/Plugins/Experimental/PCG/Source/PCG/Public/MatchAndSet/PCGMatchAndSetBase.h b/Engine/Plugins/Experimental/PCG/Source/PCG/Public/MatchAndSet/PCGMatchAndSetBase.h
--- a/Engine/Plugins/Experimental/PCG/Source/PCG/Public/MatchAndSet/PCGMatchAndSetBase.h
+++ b/Engine/Plugins/Experimental/PCG/Source/PCG/Public/MatchAndSet/PCGMatchAndSetBase.h
@@ -2,11 +2,12 @@
 
 #pragma once
 
-#include "CoreMinimal.h"
 #include "Metadata/PCGMetadataTypesConstantStruct.h"
 
 #include "PCGMatchAndSetBase.generated.h"
 
+class UPCGPointData;
+
 struct FPCGAttributePropertySelector;
 struct FPCGContext;
 struct FPCGMetadataTypesConstantStruct;
@@ -68,4 +69,8 @@ protected:
 
 	UPROPERTY()
 	EPCGMetadataTypesConstantStructStringMode StringMode = EPCGMetadataTypesConstantStructStringMode::String;
-};
\ No newline at end of file
+};
+
+#if UE_ENABLE_INCLUDE_ORDER_DEPRECATED_IN_5_2
+#include "CoreMinimal.h"
+#endif
diff --git a/Engine/Plugins/Experimental/PCG/Source/PCG/Public/MatchAndSet/PCGMatchAndSetByAttribute.h b/Engine/Plugins/Experimental/PCG/Source/PCG/Public/MatchAndSet/PCGMatchAndSetByAttribute.h
--- a/Engine/Plugins/Experimental/PCG/Source/PCG/Public/MatchAndSet/PCGMatchAndSetByAttribute.h
+++ b/Engine/Plugins/Experimental/PCG/Source/PCG/Public/MatchAndSet/PCGMatchAndSetByAttribute.h
@@ -3,10 +3,11 @@
 #pragma once
 
 #include "MatchAndSet/PCGMatchAndSetBase.h"
-#include "Metadata/PCGMetadataTypesConstantStruct.h"
 
 #include "PCGMatchAndSetByAttribute.generated.h"
 
+struct FPropertyChangedEvent;
+
 USTRUCT(BlueprintType)
 struct PCG_API FPCGMatchAndSetByAttributeEntry
 {
@@ -64,4 +65,4 @@ public:
 	/** Lookup entries (key-value pairs) */
 	UPROPERTY(BlueprintReadWrite, EditAnywhere, Category = Settings)
 	TArray<FPCGMatchAndSetByAttributeEntry> Entries;
-};
\ No newline at end of file
+};
diff --git a/Engine/Plugins/Experimental/PCG/Source/PCG/Public/MatchAndSet/PCGMatchAndSetWeighted.h b/Engine/Plugins/Experimental/PCG/Source/PCG/Public/MatchAndSet/PCGMatchAndSetWeighted.h
--- a/Engine/Plugins/Experimental/PCG/Source/PCG/Public/MatchAndSet/PCGMatchAndSetWeighted.h
+++ b/Engine/Plugins/Experimental/PCG/Source/PCG/Public/MatchAndSet/PCGMatchAndSetWeighted.h
@@ -3,10 +3,12 @@
 #pragma once
 
 #include "MatchAndSet/PCGMatchAndSetBase.h"
-#include "Metadata/PCGMetadataTypesConstantStruct.h"
 
 #include "PCGMatchAndSetWeighted.generated.h"
 
+enum class EPCGMetadataTypes : uint8;
+struct FPropertyChangedEvent;
+
 USTRUCT(BlueprintType)
 struct PCG_API FPCGMatchAndSetWeightedEntry
 {
@@ -54,4 +56,4 @@ public:
 	/** Controls whether the output data should mutate its seed - prevents issues when doing multiple random processes in a row */
 	UPROPERTY(BlueprintReadWrite, EditAnywhere, Category = Settings)
 	bool bShouldMutateSeed = true;
-};
\ No newline at end of file
+};
