MeshModelingTools: Bind the gizmo numerical UI right after initializing the gizmo helper so that any other gizmos we create in mode setup (e.g. selection system) already have the UI bound to them.

diff --git a/Engine/Plugins/Editor/ModelingToolsEditorMode/Source/ModelingToolsEditorMode/Private/ModelingToolsEditorMode.cpp b/Engine/Plugins/Editor/ModelingToolsEditorMode/Source/ModelingToolsEditorMode/Private/ModelingToolsEditorMode.cpp
--- a/Engine/Plugins/Editor/ModelingToolsEditorMode/Source/ModelingToolsEditorMode/Private/ModelingToolsEditorMode.cpp
+++ b/Engine/Plugins/Editor/ModelingToolsEditorMode/Source/ModelingToolsEditorMode/Private/ModelingToolsEditorMode.cpp
@@ -384,6 +384,12 @@ void UModelingToolsEditorMode::Enter()
 		GetInteractiveToolsContext()->SetAbsoluteWorldSnappingEnabled(ModelingModeSettings->bEnableAbsoluteWorldSnapping);
 	}
 
+	// Now that we have the gizmo helper, bind the numerical UI.
+	if (ensure(Toolkit.IsValid()))
+	{
+		((FModelingToolsEditorModeToolkit*)Toolkit.Get())->BindGizmoNumericalUI();
+	}
+	
 	// register snapping manager
 	UE::Geometry::RegisterSceneSnappingManager(GetInteractiveToolsContext());
 	SceneSnappingManager = UE::Geometry::FindModelingSceneSnappingManager(GetToolManager());
diff --git a/Engine/Plugins/Editor/ModelingToolsEditorMode/Source/ModelingToolsEditorMode/Private/ModelingToolsEditorModeToolkit.cpp b/Engine/Plugins/Editor/ModelingToolsEditorMode/Source/ModelingToolsEditorMode/Private/ModelingToolsEditorModeToolkit.cpp
--- a/Engine/Plugins/Editor/ModelingToolsEditorMode/Source/ModelingToolsEditorMode/Private/ModelingToolsEditorModeToolkit.cpp
+++ b/Engine/Plugins/Editor/ModelingToolsEditorMode/Source/ModelingToolsEditorMode/Private/ModelingToolsEditorModeToolkit.cpp
@@ -885,13 +885,6 @@ void FModelingToolsEditorModeToolkit::InitializeAfterModeSetup()
 		// ToolTarget Factories that are only available once ModelingToolsEditorMode has been initialized
 		AssetLODMode->SetSelectedItem(AssetLODModes[0]);
 
-		// Attach the gizmo numerical UI to the gizmo context object so that it can show the values for any newly
-		// created gizmos.
-		if (ensure(GizmoNumericalUIOverlayWidget.IsValid()))
-		{
-			GizmoNumericalUIOverlayWidget->BindToGizmoContextObject(GetScriptableEditorMode()->GetInteractiveToolsContext(EToolsContextScope::EdMode));
-		}
-
 		bFirstInitializeAfterModeSetup = false;
 	}
 
@@ -1512,6 +1505,14 @@ void FModelingToolsEditorModeToolkit::ForceToolPaletteRebuild()
 	//}
 }
 
+void FModelingToolsEditorModeToolkit::BindGizmoNumericalUI()
+{
+	if (ensure(GizmoNumericalUIOverlayWidget.IsValid()))
+	{
+		ensure(GizmoNumericalUIOverlayWidget->BindToGizmoContextObject(GetScriptableEditorMode()->GetInteractiveToolsContext(EToolsContextScope::EdMode)));
+	}
+}
+
 
 void FModelingToolsEditorModeToolkit::OnToolPaletteChanged(FName PaletteName) 
 {
diff --git a/Engine/Plugins/Editor/ModelingToolsEditorMode/Source/ModelingToolsEditorMode/Public/ModelingToolsEditorModeToolkit.h b/Engine/Plugins/Editor/ModelingToolsEditorMode/Source/ModelingToolsEditorMode/Public/ModelingToolsEditorModeToolkit.h
--- a/Engine/Plugins/Editor/ModelingToolsEditorMode/Source/ModelingToolsEditorMode/Public/ModelingToolsEditorModeToolkit.h
+++ b/Engine/Plugins/Editor/ModelingToolsEditorMode/Source/ModelingToolsEditorMode/Public/ModelingToolsEditorModeToolkit.h
@@ -94,6 +94,10 @@ public:
 
 	virtual void ForceToolPaletteRebuild();
 
+	// The mode is expected to call this after setting up the gizmo context object so that any subsequently
+	// created gizmos are bound to the numerical UI.
+	void BindGizmoNumericalUI();
+
 	// This is exposed only for the convenience of being able to create the numerical UI submenu
 	// in a non-member function in ModelingModeToolkit_Toolbars.cpp
 	TSharedPtr<STransformGizmoNumericalUIOverlay> GetGizmoNumericalUIOverlayWidget() { return GizmoNumericalUIOverlayWidget; }
