ModelingMode: fix PolyEd and TriSel buttons in Mesh Selection palette, they were starting the wrong tools

diff --git a/Engine/Plugins/Editor/ModelingToolsEditorMode/Source/ModelingToolsEditorMode/Private/ModelingToolsEditorMode.cpp b/Engine/Plugins/Editor/ModelingToolsEditorMode/Source/ModelingToolsEditorMode/Private/ModelingToolsEditorMode.cpp
--- a/Engine/Plugins/Editor/ModelingToolsEditorMode/Source/ModelingToolsEditorMode/Private/ModelingToolsEditorMode.cpp
+++ b/Engine/Plugins/Editor/ModelingToolsEditorMode/Source/ModelingToolsEditorMode/Private/ModelingToolsEditorMode.cpp
@@ -752,8 +752,8 @@ void UModelingToolsEditorMode::Enter()
 	RegisterSelectionTool(ToolManagerCommands.BeginSelectionAction_Extrude, TEXT("BeginSelectionExtrudeTool"), NewObject<UExtrudeMeshSelectionToolBuilder>(), true, true);
 	RegisterSelectionTool(ToolManagerCommands.BeginSelectionAction_Offset, TEXT("BeginSelectionOffsetTool"), NewObject<UOffsetMeshSelectionToolBuilder>(), true, true);
 
-	RegisterSelectionTool(ToolManagerCommands.BeginPolyModelTool_PolyEd, TEXT("BeginSelectionPolyEdTool"), NewObject<UMeshSelectionToolBuilder>(), true, false);
-	RegisterSelectionTool(ToolManagerCommands.BeginPolyModelTool_TriSel, TEXT("BeginSelectionTriEdTool"), NewObject<UOffsetMeshSelectionToolBuilder>(), true, false);
+	RegisterSelectionTool(ToolManagerCommands.BeginPolyModelTool_PolyEd, TEXT("BeginSelectionPolyEdTool"), NewObject<UEditMeshPolygonsToolBuilder>(), true, false);
+	RegisterSelectionTool(ToolManagerCommands.BeginPolyModelTool_TriSel, TEXT("BeginSelectionTriEdTool"), NewObject<UMeshSelectionToolBuilder>(), true, false);
 
 
 	auto RegisterPolyModelActionTool = [&](EEditMeshPolygonsToolActions Action, TSharedPtr<FUICommandInfo> UICommand, FString StringName, bool bRequiresSelection)
