ModelingMode: fix scaling issue in mesh selection transform gizmo

diff --git a/Engine/Plugins/Editor/ModelingToolsEditorMode/Source/ModelingToolsEditorMode/Private/ModelingSelectionInteraction.cpp b/Engine/Plugins/Editor/ModelingToolsEditorMode/Source/ModelingToolsEditorMode/Private/ModelingSelectionInteraction.cpp
--- a/Engine/Plugins/Editor/ModelingToolsEditorMode/Source/ModelingToolsEditorMode/Private/ModelingSelectionInteraction.cpp
+++ b/Engine/Plugins/Editor/ModelingToolsEditorMode/Source/ModelingToolsEditorMode/Private/ModelingSelectionInteraction.cpp
@@ -373,6 +373,9 @@ void UModelingSelectionInteraction::OnEndGizmoTransform(UTransformProxy* Proxy)
 		SelectionManager->EndTransformation();
 		bInActiveTransform = false;
 
+		// reset the transient scaling being stored by the gizmo, otherwise it will be re-applied in the next transformation
+		TransformGizmo->SetNewChildScale(FVector::One());
+
 		OnTransformEnd.Broadcast();
 	}
 }
