Add the ability for modes to add a secondary toolbar below the Level Editor toolbar, currently used by the modeling mode

diff --git a/Engine/Plugins/Editor/ModelingToolsEditorMode/Source/ModelingToolsEditorMode/ModelingToolsEditorMode.Build.cs b/Engine/Plugins/Editor/ModelingToolsEditorMode/Source/ModelingToolsEditorMode/ModelingToolsEditorMode.Build.cs
--- a/Engine/Plugins/Editor/ModelingToolsEditorMode/Source/ModelingToolsEditorMode/ModelingToolsEditorMode.Build.cs
+++ b/Engine/Plugins/Editor/ModelingToolsEditorMode/Source/ModelingToolsEditorMode/ModelingToolsEditorMode.Build.cs
@@ -49,6 +49,7 @@ public class ModelingToolsEditorMode : ModuleRules
 				"StylusInput",
 				"DeveloperSettings",
 				"PropertyEditor",
+				"ToolMenus",
 				"PresetAsset"
 				// ... add private dependencies that you statically link with here ...	
 			}
diff --git a/Engine/Plugins/Editor/ModelingToolsEditorMode/Source/ModelingToolsEditorMode/Private/ModelingModeToolkit_Toolbars.cpp b/Engine/Plugins/Editor/ModelingToolsEditorMode/Source/ModelingToolsEditorMode/Private/ModelingModeToolkit_Toolbars.cpp
--- a/Engine/Plugins/Editor/ModelingToolsEditorMode/Source/ModelingToolsEditorMode/Private/ModelingModeToolkit_Toolbars.cpp
+++ b/Engine/Plugins/Editor/ModelingToolsEditorMode/Source/ModelingToolsEditorMode/Private/ModelingModeToolkit_Toolbars.cpp
@@ -18,9 +18,14 @@
 
 #include "Framework/MultiBox/MultiBoxBuilder.h"
 #include "STransformGizmoNumericalUIOverlay.h"
+#include "Widgets/Images/SImage.h"
+#include "Widgets/Layout/SBorder.h"
+#include "Widgets/Images/SImage.h"
+#include "Widgets/Layout/SBorder.h"
 #include "Widgets/Layout/SSpacer.h"
 #include "Widgets/Input/SComboButton.h"
 
+#include "ToolMenus.h"
 // to toggle component instance selection
 
 #define LOCTEXT_NAMESPACE "FModelingToolsEditorModeToolkit"
@@ -263,44 +268,6 @@ void MakeSubMenu_DefaultMeshObjectType(FMenuBuilder& MenuBuilder)
 	MakeSubMenu_DefaultMeshObjectPhysicsSettings(MenuBuilder);
 }
 
-
-
-TSharedRef<SWidget> MakeMenu_ModelingModeConfigSettings(FModelingToolsEditorModeToolkit* Toolkit)
-{
-	FMenuBuilder MenuBuilder(true, TSharedPtr<FUICommandList>());
-
-	MenuBuilder.BeginSection("Section_Selection", LOCTEXT("Section_Selection", "Selection"));
-	MakeSubMenu_ModeToggles(MenuBuilder);
-	MenuBuilder.EndSection();
-
-	MenuBuilder.BeginSection("Section_Gizmo", LOCTEXT("Section_Gizmo", "Gizmo"));
-	MakeSubMenu_GizmoVisibilityMode(Toolkit, MenuBuilder);
-	MenuBuilder.EndSection();
-
-	MenuBuilder.BeginSection("Section_MeshObjects", LOCTEXT("Section_MeshObjects", "New Mesh Objects"));
-	MenuBuilder.AddSubMenu(
-		LOCTEXT("MeshObjectTypeSubMenu", "New Mesh Settings"), LOCTEXT("MeshObjectTypeSubMenu_ToolTip", "Configure default settings for new Mesh Object Types"),
-		FNewMenuDelegate::CreateLambda([=](FMenuBuilder& SubMenuBuilder) {
-			MakeSubMenu_DefaultMeshObjectType(SubMenuBuilder);
-		}));
-	MenuBuilder.EndSection();
-
-	MenuBuilder.BeginSection("Section_Settings", LOCTEXT("Section_Settings", "Quick Settings"));
-	MenuBuilder.AddSubMenu(
-		LOCTEXT("QuickSettingsSubMenu", "Jump To Settings"), LOCTEXT("QuickSettingsSubMenu_ToolTip", "Jump to sections of the Settings dialogs relevant to Modeling Mode"),
-		FNewMenuDelegate::CreateLambda([=](FMenuBuilder& SubMenuBuilder) {
-			MakeSubMenu_QuickSettings(SubMenuBuilder);
-			}));
-	MenuBuilder.EndSection();
-
-	TSharedRef<SWidget> MenuWidget = MenuBuilder.MakeWidget();
-	return MenuWidget;
-}
-
-
-
-
-
 void MakeSubMenu_Selection_DragMode(FModelingToolsEditorModeToolkit* Toolkit, FMenuBuilder& MenuBuilder)
 {
 	UModelingSelectionInteraction* SelectionInteraction = Cast<UModelingToolsEditorMode>(Toolkit->GetScriptableEditorMode())->GetSelectionInteraction();
@@ -384,9 +351,7 @@ void MakeSubMenu_Selection_MeshType(FModelingToolsEditorModeToolkit* Toolkit, FM
 
 }
 
-
-
-TSharedRef<SWidget> MakeMenu_SelectionConfigSettings(FModelingToolsEditorModeToolkit* Toolkit)
+TSharedRef<SWidget> MakeMenu_ModelingModeConfigSettings(FModelingToolsEditorModeToolkit* Toolkit)
 {
 	FMenuBuilder MenuBuilder(true, TSharedPtr<FUICommandList>());
 
@@ -398,94 +363,125 @@ TSharedRef<SWidget> MakeMenu_SelectionConfigSettings(FModelingToolsEditorModeToo
 	MakeSubMenu_Selection_MeshType(Toolkit, MenuBuilder);
 	MenuBuilder.EndSection();
 
-	TSharedRef<SWidget> MenuWidget = MenuBuilder.MakeWidget();
-	return MenuWidget;
-}
-
-
-
-
-TSharedRef<SWidget> MakeMenu_SelectionEdits(FModelingToolsEditorModeToolkit* Toolkit)
-{
-	FMenuBuilder MenuBuilder(true, Toolkit->GetToolkitCommands() );
+	MenuBuilder.BeginSection("Section_Selection", LOCTEXT("Section_Selection", "Selection"));
+	MakeSubMenu_ModeToggles(MenuBuilder);
+	MenuBuilder.EndSection();
 
-	MenuBuilder.BeginSection("Section_SelectionEdits", LOCTEXT("Section_SelectionEdits", "Selection Edits"));
+	MenuBuilder.BeginSection("Section_Gizmo", LOCTEXT("Section_Gizmo", "Gizmo"));
+	MakeSubMenu_GizmoVisibilityMode(Toolkit, MenuBuilder);
+	MenuBuilder.EndSection();
 
-	const FModelingToolsManagerCommands& Commands = FModelingToolsManagerCommands::Get();
+	MenuBuilder.BeginSection("Section_MeshObjects", LOCTEXT("Section_MeshObjects", "New Mesh Objects"));
+	MenuBuilder.AddSubMenu(
+		LOCTEXT("MeshObjectTypeSubMenu", "New Mesh Settings"), LOCTEXT("MeshObjectTypeSubMenu_ToolTip", "Configure default settings for new Mesh Object Types"),
+		FNewMenuDelegate::CreateLambda([=](FMenuBuilder& SubMenuBuilder) {
+			MakeSubMenu_DefaultMeshObjectType(SubMenuBuilder);
+		}));
+	MenuBuilder.EndSection();
 
-	MenuBuilder.AddMenuEntry(Commands.BeginSelectionAction_SelectAll);
-	MenuBuilder.AddMenuEntry(Commands.BeginSelectionAction_ExpandToConnected);
-	MenuBuilder.AddMenuEntry(Commands.BeginSelectionAction_Invert);
-	MenuBuilder.AddMenuEntry(Commands.BeginSelectionAction_InvertConnected);
-	MenuBuilder.AddMenuEntry(Commands.BeginSelectionAction_Expand);
-	MenuBuilder.AddMenuEntry(Commands.BeginSelectionAction_Contract);
+	MenuBuilder.BeginSection("Section_Settings", LOCTEXT("Section_Settings", "Quick Settings"));
+	MenuBuilder.AddSubMenu(
+		LOCTEXT("QuickSettingsSubMenu", "Jump To Settings"), LOCTEXT("QuickSettingsSubMenu_ToolTip", "Jump to sections of the Settings dialogs relevant to Modeling Mode"),
+		FNewMenuDelegate::CreateLambda([=](FMenuBuilder& SubMenuBuilder) {
+			MakeSubMenu_QuickSettings(SubMenuBuilder);
+			}));
+	MenuBuilder.EndSection();
 
 	TSharedRef<SWidget> MenuWidget = MenuBuilder.MakeWidget();
 	return MenuWidget;
 }
 
-
 } // end namespace UELocal
 
 
-void FModelingToolsEditorModeToolkit::MakeSelectionPaletteOverlayWidget()
+void FModelingToolsEditorModeToolkit::ExtendSecondaryModeToolbar(UToolMenu *InModeToolbarMenu)
 {
-	FVerticalToolBarBuilder ToolbarBuilder( 
-		GetToolkitCommands(), 
-		FMultiBoxCustomization::None, 
-		TSharedPtr<FExtender>(), /*InForceSmallIcons=*/ true);
-	ToolbarBuilder.SetLabelVisibility(EVisibility::Collapsed);
-	ToolbarBuilder.SetStyle(&FAppStyle::Get(), "EditorViewportToolBar");
+	/*
+	 * Sanity check because the toolbar extension should happen after the ModeUILayer setup which is currently manual
+	 * for standalone asset editors
+	 */
+	if( !ModeUILayer.IsValid() )
+	{
+		return;
+	}
+
+	UModelingToolsEditorModeSettings* ModelingModeSettings = GetMutableDefault<UModelingToolsEditorModeSettings>();
+	const bool bEnableSelectionUI = ModelingModeSettings && ModelingModeSettings->bEnableMeshSelections;
+	if ( !bEnableSelectionUI )
+	{
+		return;
+	}
+	
+	FToolMenuSection& Section = InModeToolbarMenu->FindOrAddSection("SelectionPalette");
 
+	/*
+	 * NOTE: Any commands used in the SecondaryModeToolbar need to be present in the command list that lives in the
+	 * ModeUILayer. Currently this is done automatically by parenting said command list to the toolkit's command list
+	 */
 	const FModelingToolsManagerCommands& Commands = FModelingToolsManagerCommands::Get();
 
-	ToolbarBuilder.AddToolBarButton(Commands.MeshSelectionModeAction_NoSelection);
-	ToolbarBuilder.AddWidget( SNew(SSpacer).Size(FVector2D(1,4)) );
-	ToolbarBuilder.AddToolBarButton(Commands.MeshSelectionModeAction_MeshTriangles);
-	ToolbarBuilder.AddToolBarButton(Commands.MeshSelectionModeAction_MeshEdges);
-	ToolbarBuilder.AddToolBarButton(Commands.MeshSelectionModeAction_MeshVertices);
-	ToolbarBuilder.AddWidget( SNew(SSpacer).Size(FVector2D(1,4)) );
-	ToolbarBuilder.AddToolBarButton(Commands.MeshSelectionModeAction_GroupFaces);
-	ToolbarBuilder.AddToolBarButton(Commands.MeshSelectionModeAction_GroupEdges);
-	ToolbarBuilder.AddToolBarButton(Commands.MeshSelectionModeAction_GroupCorners);
-
-	ToolbarBuilder.AddWidget(
-		SNew(SComboButton)
-		.HasDownArrow(false)
-		.MenuPlacement(EMenuPlacement::MenuPlacement_MenuRight)
-		.ComboButtonStyle(FAppStyle::Get(), "SimpleComboButton")
-		.OnGetMenuContent(FOnGetContent::CreateLambda([this]()
-		{
-			return UELocal::MakeMenu_SelectionEdits(this);
-		}))
-		.ContentPadding(FMargin(1.0f, 1.0f))
-		.ButtonContent()
+	Section.AddEntry(FToolMenuEntry::InitToolBarButton(Commands.MeshSelectionModeAction_NoSelection));
+
+	Section.AddSeparator(NAME_None);
+
+	TSharedRef<SWidget> TrianglesTextWidget = SNew(SHorizontalBox)
+		+ SHorizontalBox::Slot()
+		.AutoWidth()
+		.Padding(8.0f, 0.0f, 8.0f, 0.0f)
+		.VAlign(VAlign_Center)
 		[
-			SNew(SImage)
-			.Image(FModelingToolsEditorModeStyle::Get()->GetBrush("ModelingModeSelection.Edits_Right"))
-		] );
+			SNew(STextBlock)
+			.Text(LOCTEXT("TriangleSelectionText", "Triangles"))
+		];
 
-	ToolbarBuilder.AddWidget(
-		SNew(SComboButton)
-		.HasDownArrow(false)
-		.MenuPlacement(EMenuPlacement::MenuPlacement_MenuRight)
-		.ComboButtonStyle(FAppStyle::Get(), "SimpleComboButton")
-		.OnGetMenuContent(FOnGetContent::CreateLambda([this]()
-		{
-			return UELocal::MakeMenu_SelectionConfigSettings(this);
-		}))
-		.ContentPadding(FMargin(1.0f, 1.0f))
-		.ButtonContent()
+	Section.AddEntry(FToolMenuEntry::InitWidget("TrianglesTextWidget", TrianglesTextWidget, FText()));
+	
+	Section.AddEntry(FToolMenuEntry::InitToolBarButton(Commands.MeshSelectionModeAction_MeshTriangles));
+	Section.AddEntry(FToolMenuEntry::InitToolBarButton(Commands.MeshSelectionModeAction_MeshEdges));
+	Section.AddEntry(FToolMenuEntry::InitToolBarButton(Commands.MeshSelectionModeAction_MeshVertices));
+
+	Section.AddSeparator(NAME_None);
+
+	TSharedRef<SWidget> PolygroupsTextWidget = SNew(SHorizontalBox)
+		+ SHorizontalBox::Slot()
+		.AutoWidth()
+		.Padding(8.0f, 0.0f, 8.0f, 0.0f)
+		.VAlign(VAlign_Center)
 		[
-			SNew(SImage)
-			.Image(FModelingToolsEditorModeStyle::Get()->GetBrush("ModelingModeSelection.More_Right"))
-		] );
+			SNew(STextBlock)
+			.Text(LOCTEXT("PolygroupsSelectionText", "Polygroups"))
+		];
 
+	Section.AddEntry(FToolMenuEntry::InitWidget("PolygroupsTextWidget", PolygroupsTextWidget, FText()));
+	
+	Section.AddEntry(FToolMenuEntry::InitToolBarButton(Commands.MeshSelectionModeAction_GroupFaces));
+	Section.AddEntry(FToolMenuEntry::InitToolBarButton(Commands.MeshSelectionModeAction_GroupEdges));
+	Section.AddEntry(FToolMenuEntry::InitToolBarButton(Commands.MeshSelectionModeAction_GroupCorners));
 
-	ToolbarBuilder.AddWidget( SNew(SSpacer).Size(FVector2D(1,10)) );
+	Section.AddSeparator(NAME_None);
 
-	ToolbarBuilder.AddWidget(
-		SNew(SComboButton)
+	TSharedRef<SWidget> SelectionEditTextWidget = SNew(SHorizontalBox)
+		+ SHorizontalBox::Slot()
+		.AutoWidth()
+		.Padding(8.0f, 0.0f, 8.0f, 0.0f)
+		.VAlign(VAlign_Center)
+		[
+			SNew(STextBlock)
+			.Text(LOCTEXT("SelectionEditText", "Selection Edit"))
+		];
+
+	Section.AddEntry(FToolMenuEntry::InitWidget("SelectionEditTextWidget", SelectionEditTextWidget, FText()));
+
+	Section.AddEntry(FToolMenuEntry::InitToolBarButton(Commands.BeginSelectionAction_SelectAll));
+	Section.AddEntry(FToolMenuEntry::InitToolBarButton(Commands.BeginSelectionAction_Invert));
+	Section.AddEntry(FToolMenuEntry::InitToolBarButton(Commands.BeginSelectionAction_ExpandToConnected));
+	Section.AddEntry(FToolMenuEntry::InitToolBarButton(Commands.BeginSelectionAction_InvertConnected));
+	Section.AddEntry(FToolMenuEntry::InitToolBarButton(Commands.BeginSelectionAction_Expand));
+	Section.AddEntry(FToolMenuEntry::InitToolBarButton(Commands.BeginSelectionAction_Contract));
+	
+	Section.AddSeparator(NAME_None);
+
+	TSharedRef<SWidget> DefaultSettingsWidget = SNew(SComboButton)
 		.HasDownArrow(false)
 		.MenuPlacement(EMenuPlacement::MenuPlacement_MenuRight)
 		.ComboButtonStyle(FAppStyle::Get(), "SimpleComboButton")
@@ -498,34 +494,13 @@ void FModelingToolsEditorModeToolkit::MakeSelectionPaletteOverlayWidget()
 		[
 			SNew(SImage)
 			.Image(FModelingToolsEditorModeStyle::Get()->GetBrush("ModelingMode.DefaultSettings"))
-		] );
-
+			.ColorAndOpacity(FSlateColor::UseForeground())
+		];
 
-	TSharedRef<SWidget> Toolbar = ToolbarBuilder.MakeWidget();
-
-
-	SAssignNew(SelectionPaletteOverlayWidget, SVerticalBox)
-	+SVerticalBox::Slot()
-	.HAlign(HAlign_Left)
-	.VAlign(VAlign_Center)
-	.Padding(FMargin(2.0f, 0.0f, 0.f, 0.f))
-	[
-		SNew(SBorder)
-		.BorderImage(FAppStyle::Get().GetBrush("EditorViewport.OverlayBrush"))
-		.Padding(FMargin(3.0f, 6.0f, 3.f, 6.f))
-		[
-			SNew(SVerticalBox)
-
-			+ SVerticalBox::Slot()
-			.AutoHeight()
-			.Padding(FMargin(0.0, 0.f, 0.f, 0.f))
-			[
-				Toolbar
-			]
-		]
-	];
+	Section.AddEntry(FToolMenuEntry::InitWidget("DefaultSettings", DefaultSettingsWidget, FText()));
+	
+	Section.AddSeparator(NAME_None);
 
 }
 
-
 #undef LOCTEXT_NAMESPACE
diff --git a/Engine/Plugins/Editor/ModelingToolsEditorMode/Source/ModelingToolsEditorMode/Private/ModelingToolsEditorModeStyle.cpp b/Engine/Plugins/Editor/ModelingToolsEditorMode/Source/ModelingToolsEditorMode/Private/ModelingToolsEditorModeStyle.cpp
--- a/Engine/Plugins/Editor/ModelingToolsEditorMode/Source/ModelingToolsEditorMode/Private/ModelingToolsEditorModeStyle.cpp
+++ b/Engine/Plugins/Editor/ModelingToolsEditorMode/Source/ModelingToolsEditorMode/Private/ModelingToolsEditorModeStyle.cpp
@@ -321,14 +321,22 @@ void FModelingToolsEditorModeStyle::Initialize()
 		StyleSet->Set("ModelingToolsManagerCommands.BeginSkinWeightsPaintTool", new IMAGE_BRUSH_SVG("Icons/SkinWeightsPaint", Icon20x20));
 		StyleSet->Set("ModelingToolsManagerCommands.BeginSkinWeightsPaintTool.Small", new IMAGE_BRUSH_SVG("Icons/SkinWeightsPaint", Icon20x20));
 
-		StyleSet->Set("ModelingToolsManagerCommands.MeshSelectionModeAction_NoSelection", new FSlateImageBrush(StyleSet->RootToCoreContentDir(TEXT("../Editor/Slate/Icons/GeneralTools/Select_40x.png")), Icon20x20));
-		StyleSet->Set("ModelingToolsManagerCommands.MeshSelectionModeAction_MeshTriangles", new IMAGE_BRUSH_SVG("Icons/SelectionToolbar_Triangles", Icon20x20));
-		StyleSet->Set("ModelingToolsManagerCommands.MeshSelectionModeAction_MeshVertices", new IMAGE_BRUSH_SVG("Icons/SelectionToolbar_Vertices", Icon20x20));
-		StyleSet->Set("ModelingToolsManagerCommands.MeshSelectionModeAction_MeshEdges", new IMAGE_BRUSH_SVG("Icons/SelectionToolbar_Edges", Icon20x20));
-		StyleSet->Set("ModelingToolsManagerCommands.MeshSelectionModeAction_GroupFaces", new IMAGE_BRUSH_SVG("Icons/SelectionTriangles3", Icon20x20));
-		StyleSet->Set("ModelingToolsManagerCommands.MeshSelectionModeAction_GroupCorners", new IMAGE_BRUSH_SVG("Icons/SelectionVertices", Icon20x20));
-		StyleSet->Set("ModelingToolsManagerCommands.MeshSelectionModeAction_GroupEdges", new IMAGE_BRUSH_SVG("Icons/SelectionBorderEdges", Icon20x20));
-
+		StyleSet->Set("ModelingToolsManagerCommands.MeshSelectionModeAction_NoSelection", new IMAGE_BRUSH_SVG("Icons/ModSelectionObject_16", Icon20x20));
+		StyleSet->Set("ModelingToolsManagerCommands.MeshSelectionModeAction_MeshTriangles", new IMAGE_BRUSH_SVG("Icons/ModSelectionPolys_16", Icon20x20));
+		StyleSet->Set("ModelingToolsManagerCommands.MeshSelectionModeAction_MeshVertices", new IMAGE_BRUSH_SVG("Icons/ModSelectionVerts_16", Icon20x20));
+		StyleSet->Set("ModelingToolsManagerCommands.MeshSelectionModeAction_MeshEdges", new IMAGE_BRUSH_SVG("Icons/ModSelectionEdges_16", Icon20x20));
+		StyleSet->Set("ModelingToolsManagerCommands.MeshSelectionModeAction_GroupFaces", new IMAGE_BRUSH_SVG("Icons/ModSelectionPolygroupFaces_16", Icon20x20));
+		StyleSet->Set("ModelingToolsManagerCommands.MeshSelectionModeAction_GroupCorners", new IMAGE_BRUSH_SVG("Icons/ModSelectionPolygroupVerts_16", Icon20x20));
+		StyleSet->Set("ModelingToolsManagerCommands.MeshSelectionModeAction_GroupEdges", new IMAGE_BRUSH_SVG("Icons/ModSelectionPolygroupEdges_16", Icon20x20));
+
+		StyleSet->Set("ModelingToolsManagerCommands.BeginSelectionAction_SelectAll", new IMAGE_BRUSH_SVG("Icons/ModSelectionAll_16", Icon20x20));
+		StyleSet->Set("ModelingToolsManagerCommands.BeginSelectionAction_Invert", new IMAGE_BRUSH_SVG("Icons/ModSelectionInverse_16", Icon20x20));
+		StyleSet->Set("ModelingToolsManagerCommands.BeginSelectionAction_ExpandToConnected", new IMAGE_BRUSH_SVG("Icons/ModSelectionConnected_16", Icon20x20));
+		StyleSet->Set("ModelingToolsManagerCommands.BeginSelectionAction_InvertConnected", new IMAGE_BRUSH_SVG("Icons/ModSelectionConnectedInverse_16", Icon20x20));
+		StyleSet->Set("ModelingToolsManagerCommands.BeginSelectionAction_Expand", new IMAGE_BRUSH_SVG("Icons/ModSelectionExpand_16", Icon20x20));
+		StyleSet->Set("ModelingToolsManagerCommands.BeginSelectionAction_Contract", new IMAGE_BRUSH_SVG("Icons/ModSelectionShrink_16", Icon20x20));
+
+		
 		StyleSet->Set("ModelingToolsManagerCommands.BeginSelectionAction_Delete", new FSlateImageBrush(StyleSet->RootToCoreContentDir(TEXT("../Editor/Slate/Icons/GeneralTools/Delete_40x.png")), Icon20x20));
 
 		StyleSet->Set("ModelingToolsManagerCommands.BeginPolyModelTool_PolyEd", 				new IMAGE_PLUGIN_BRUSH("Icons/PolyEdit_40x", 		Icon20x20));
diff --git a/Engine/Plugins/Editor/ModelingToolsEditorMode/Source/ModelingToolsEditorMode/Private/ModelingToolsEditorModeToolkit.cpp b/Engine/Plugins/Editor/ModelingToolsEditorMode/Source/ModelingToolsEditorMode/Private/ModelingToolsEditorModeToolkit.cpp
--- a/Engine/Plugins/Editor/ModelingToolsEditorMode/Source/ModelingToolsEditorMode/Private/ModelingToolsEditorModeToolkit.cpp
+++ b/Engine/Plugins/Editor/ModelingToolsEditorMode/Source/ModelingToolsEditorMode/Private/ModelingToolsEditorModeToolkit.cpp
@@ -47,6 +47,7 @@
 #include "LevelEditorViewport.h"
 #include "Widgets/Input/SEditableTextBox.h"
 #include "Widgets/Input/STextComboBox.h"
+#include "ToolMenus.h"
 
 
 #define LOCTEXT_NAMESPACE "FModelingToolsEditorModeToolkit"
@@ -112,19 +113,12 @@ FModelingToolsEditorModeToolkit::~FModelingToolsEditorModeToolkit()
 {
 	if (IsHosted())
 	{
-		if (SelectionPaletteOverlayWidget.IsValid())
-		{
-			GetToolkitHost()->RemoveViewportOverlayWidget(SelectionPaletteOverlayWidget.ToSharedRef());
-			SelectionPaletteOverlayWidget.Reset();
-		}
-
 		if (GizmoNumericalUIOverlayWidget.IsValid())
 		{
 			GetToolkitHost()->RemoveViewportOverlayWidget(GizmoNumericalUIOverlayWidget.ToSharedRef());
 			GizmoNumericalUIOverlayWidget.Reset();
 		}
 	}
-
 	UModelingToolsEditorModeSettings* Settings = GetMutableDefault<UModelingToolsEditorModeSettings>();
 	Settings->OnModified.Remove(AssetSettingsModifiedHandle);
 	GetScriptableEditorMode()->GetInteractiveToolsContext(EToolsContextScope::EdMode)->OnToolNotificationMessage.RemoveAll(this);
@@ -233,14 +227,6 @@ void FModelingToolsEditorModeToolkit::Init(const TSharedPtr<IToolkitHost>& InitT
 	// it can be bound to the buttons there.
 	MakeGizmoNumericalUIOverlayWidget();
 	GetToolkitHost()->AddViewportOverlayWidget(GizmoNumericalUIOverlayWidget.ToSharedRef());
-
-	const UModelingToolsEditorModeSettings* ModelingModeSettings = GetDefault<UModelingToolsEditorModeSettings>();
-	bool bEnableSelectionUI = ModelingModeSettings && ModelingModeSettings->bEnableMeshSelections;
-	if ( bEnableSelectionUI )
-	{
-		MakeSelectionPaletteOverlayWidget();
-		GetToolkitHost()->AddViewportOverlayWidget(SelectionPaletteOverlayWidget.ToSharedRef());
-	}
 }
 
 void FModelingToolsEditorModeToolkit::AddReferencedObjects(FReferenceCollector& Collector)
@@ -969,8 +955,6 @@ void FModelingToolsEditorModeToolkit::ClearWarning()
 	ToolWarningArea->SetVisibility(EVisibility::Collapsed);
 }
 
-
-
 FName FModelingToolsEditorModeToolkit::GetToolkitFName() const
 {
 	return FName("ModelingToolsEditorMode");
diff --git a/Engine/Plugins/Editor/ModelingToolsEditorMode/Source/ModelingToolsEditorMode/Public/ModelingToolsEditorModeToolkit.h b/Engine/Plugins/Editor/ModelingToolsEditorMode/Source/ModelingToolsEditorMode/Public/ModelingToolsEditorModeToolkit.h
--- a/Engine/Plugins/Editor/ModelingToolsEditorMode/Source/ModelingToolsEditorMode/Public/ModelingToolsEditorModeToolkit.h
+++ b/Engine/Plugins/Editor/ModelingToolsEditorMode/Source/ModelingToolsEditorMode/Public/ModelingToolsEditorModeToolkit.h
@@ -87,6 +87,7 @@ public:
 	void OnActiveViewportChanged(TSharedPtr<IAssetViewport>, TSharedPtr<IAssetViewport> );
 
 	virtual void InvokeUI() override;
+	virtual void ExtendSecondaryModeToolbar(UToolMenu *InModeToolbarMenu) override;
 
 	// Read and write Autogen Subfolder quick setting in restrictive mode.
 	virtual FText GetRestrictiveModeAutoGeneratedAssetPathText() const;
@@ -120,10 +121,7 @@ private:
 
 	TSharedPtr<SWidget> ToolShutdownViewportOverlayWidget;
 	void MakeToolShutdownOverlayWidget();
-
-	TSharedPtr<SWidget> SelectionPaletteOverlayWidget;
-	void MakeSelectionPaletteOverlayWidget();
-
+	
 	void MakeGizmoNumericalUIOverlayWidget();
 	TSharedPtr<STransformGizmoNumericalUIOverlay> GizmoNumericalUIOverlayWidget;
 
diff --git a/Engine/Source/Editor/EditorFramework/EditorFramework.Build.cs b/Engine/Source/Editor/EditorFramework/EditorFramework.Build.cs
--- a/Engine/Source/Editor/EditorFramework/EditorFramework.Build.cs
+++ b/Engine/Source/Editor/EditorFramework/EditorFramework.Build.cs
@@ -15,7 +15,8 @@ public class EditorFramework : ModuleRules
 				"RenderCore",
 				"Slate",
 				"SlateCore",
-				"WorkspaceMenuStructure"
+				"WorkspaceMenuStructure",
+				"ToolMenus"
 			}
 		);
 
diff --git a/Engine/Source/Editor/EditorFramework/Private/Toolkits/AssetEditorModeUILayer.cpp b/Engine/Source/Editor/EditorFramework/Private/Toolkits/AssetEditorModeUILayer.cpp
--- a/Engine/Source/Editor/EditorFramework/Private/Toolkits/AssetEditorModeUILayer.cpp
+++ b/Engine/Source/Editor/EditorFramework/Private/Toolkits/AssetEditorModeUILayer.cpp
@@ -9,6 +9,7 @@
 #include "Widgets/DeclarativeSyntaxSupport.h"
 #include "Widgets/Docking/SDockTab.h"
 #include "Widgets/SNullWidget.h"
+#include "Toolkits/IToolkit.h"
 
 #define LOCTEXT_NAMESPACE "SLevelEditorToolBox"
 
@@ -20,6 +21,8 @@ const FName UAssetEditorUISubsystem::VerticalToolbarID = TEXT("VerticalModeToolb
 
 FAssetEditorModeUILayer::FAssetEditorModeUILayer(const IToolkitHost* InToolkitHost)
 	: ToolkitHost(InToolkitHost)
+	, ModeCommands( new FUICommandList() )
+
 {
 	RequestedTabInfo.Add(UAssetEditorUISubsystem::VerticalToolbarID, FMinorTabConfig(UAssetEditorUISubsystem::VerticalToolbarID));
 	RequestedTabInfo.Add(UAssetEditorUISubsystem::TopLeftTabID, FMinorTabConfig(UAssetEditorUISubsystem::TopLeftTabID));
@@ -35,7 +38,6 @@ void FAssetEditorModeUILayer::OnToolkitHostingStarted(const TSharedRef<IToolkit>
 	GetTabManager()->UnregisterTabSpawner(UAssetEditorUISubsystem::VerticalToolbarID);
 	GetTabManager()->UnregisterTabSpawner(UAssetEditorUISubsystem::TopRightTabID);
 	GetTabManager()->UnregisterTabSpawner(UAssetEditorUISubsystem::BottomRightTabID);
-
 }
 
 void FAssetEditorModeUILayer::OnToolkitHostingFinished(const TSharedRef<IToolkit>& Toolkit)
@@ -100,6 +102,11 @@ void FAssetEditorModeUILayer::SetModePanelInfo(const FName InTabSpawnerID, const
 	RequestedTabInfo.Emplace(InTabSpawnerID, InTabInfo);
 }
 
+void FAssetEditorModeUILayer::SetSecondaryModeToolbarName(FName InName)
+{
+	SecondaryModeToolbarName = InName;
+}
+
 TMap<FName, TWeakPtr<SDockTab>> FAssetEditorModeUILayer::GetSpawnedTabs()
 {
 	return SpawnedTabs;
diff --git a/Engine/Source/Editor/EditorFramework/Public/Toolkits/AssetEditorModeUILayer.h b/Engine/Source/Editor/EditorFramework/Public/Toolkits/AssetEditorModeUILayer.h
--- a/Engine/Source/Editor/EditorFramework/Public/Toolkits/AssetEditorModeUILayer.h
+++ b/Engine/Source/Editor/EditorFramework/Public/Toolkits/AssetEditorModeUILayer.h
@@ -24,6 +24,7 @@
 #include "UObject/UObjectGlobals.h"
 #include "Widgets/DeclarativeSyntaxSupport.h"
 #include "Widgets/SCompoundWidget.h"
+#include "ToolMenuDelegates.h"
 
 #include "AssetEditorModeUILayer.generated.h"
 
@@ -68,6 +69,10 @@ public:
 	virtual TSharedPtr<FWorkspaceItem> GetModeMenuCategory() const;
 	virtual void SetModePanelInfo(const FName InTabSpawnerID, const FMinorTabConfig& InTabInfo);
 	virtual TMap<FName, TWeakPtr<SDockTab>> GetSpawnedTabs();
+
+	/** Called by the Toolkit Host to set the name of the ModeToolbar this mode will extend */
+	virtual void SetSecondaryModeToolbarName(FName InName);
+	
 	virtual FSimpleDelegate& ToolkitHostReadyForUI()
 	{
 		return OnToolkitHostReadyForUI;
@@ -80,6 +85,19 @@ public:
 	{
 		return ToolkitHost->GetStatusBarName();
 	}
+	virtual const FName GetSecondaryModeToolbarName() const
+	{
+		return SecondaryModeToolbarName;
+	}
+	virtual const TSharedRef<FUICommandList> GetModeCommands() const
+	{
+		return ModeCommands;
+	}
+	/* Called by the mode toolkit to extend the toolbar */
+	virtual FNewToolMenuDelegate& RegisterSecondaryModeToolbarExtension()
+	{
+		return OnRegisterSecondaryModeToolbarExtension;
+	}
 
 protected:
 	const FOnSpawnTab& GetStoredSpawner(const FName TabID);
@@ -100,4 +118,13 @@ protected:
 	TMap<FName, TWeakPtr<SDockTab>> SpawnedTabs;
 	FSimpleDelegate OnToolkitHostReadyForUI;
 	FSimpleDelegate OnToolkitHostShutdownUI;
+
+	/* The name of the mode toolbar that this mode will extend (appears below the main toolbar) */
+	FName SecondaryModeToolbarName;
+
+	/* Delegate called to actually extend the mode toolbar */
+	FNewToolMenuDelegate OnRegisterSecondaryModeToolbarExtension;
+
+	/* A list of commands this ModeUILayer is aware of, currently only passed into the Mode Toolbar */
+	TSharedRef<FUICommandList> ModeCommands;
 };
diff --git a/Engine/Source/Editor/LevelEditor/Private/LevelEditorToolBar.cpp b/Engine/Source/Editor/LevelEditor/Private/LevelEditorToolBar.cpp
--- a/Engine/Source/Editor/LevelEditor/Private/LevelEditorToolBar.cpp
+++ b/Engine/Source/Editor/LevelEditor/Private/LevelEditorToolBar.cpp
@@ -64,6 +64,8 @@
 #include "PlatformInfo.h"
 #include "DataDrivenShaderPlatformInfo.h"
 
+FName FLevelEditorToolBar::SecondaryModeToolbarName("LevelEditor.SecondaryToolbar");
+
 namespace PreviewModeFunctionality
 {
 	FText GetPreviewModeText()
@@ -1509,6 +1511,36 @@ void FLevelEditorToolBar::RegisterLevelEditorToolBar( const TSharedRef<FUIComman
 #undef LOCTEXT_NAMESPACE
 }
 
+FName FLevelEditorToolBar::GetSecondaryModeToolbarName()
+{
+	return SecondaryModeToolbarName;
+}
+
+TSharedRef< SWidget > FLevelEditorToolBar::MakeLevelEditorSecondaryModeToolbar( TSharedRef<FUICommandList> InCommandList, TMap<FName, TSharedPtr<FLevelEditorModeUILayer>>& ModeUILayers )
+{
+	FToolMenuContext MenuContext(InCommandList);
+
+	for(const TPair<FName, TSharedPtr<FLevelEditorModeUILayer>>& ModeUILayer : ModeUILayers)
+	{
+		MenuContext.AppendCommandList(ModeUILayer.Value->GetModeCommands());
+	}
+
+	return SNew(SBorder)
+	.Padding(0)
+	.BorderImage(FAppStyle::Get().GetBrush("NoBorder"))
+	.IsEnabled(FSlateApplication::Get().GetNormalExecutionAttribute())
+	[
+		UToolMenus::Get()->GenerateWidget(SecondaryModeToolbarName, MenuContext)
+	];
+
+}
+
+void FLevelEditorToolBar::RegisterLevelEditorSecondaryModeToolbar()
+{
+	UToolMenu* ModesToolbar = UToolMenus::Get()->RegisterMenu(SecondaryModeToolbarName, NAME_None, EMultiBoxType::SlimHorizontalToolBar);
+	ModesToolbar->StyleName = "SecondaryToolbar";
+}
+
 /**
  * Static: Creates a widget for the level editor tool bar
  *
diff --git a/Engine/Source/Editor/LevelEditor/Private/LevelEditorToolBar.h b/Engine/Source/Editor/LevelEditor/Private/LevelEditorToolBar.h
--- a/Engine/Source/Editor/LevelEditor/Private/LevelEditorToolBar.h
+++ b/Engine/Source/Editor/LevelEditor/Private/LevelEditorToolBar.h
@@ -24,9 +24,17 @@ public:
 	 * @return	New widget
 	 */
 	static TSharedRef< SWidget > MakeLevelEditorToolBar( const TSharedRef<FUICommandList>& InCommandList, const TSharedRef<SLevelEditor> InLevelEditor );
-
 	static void RegisterLevelEditorToolBar( const TSharedRef<FUICommandList>& InCommandList, const TSharedRef<SLevelEditor> InLevelEditor );
 
+	/**
+	 * Static: Creates a widget for the secondary tool bar which is displayed below the main toolbar
+	 *
+	 * @return	New widget
+	 */
+	static TSharedRef< SWidget > MakeLevelEditorSecondaryModeToolbar(TSharedRef<FUICommandList> InCommandList, TMap<FName, TSharedPtr<FLevelEditorModeUILayer>>& ModeUILayers );
+	static void RegisterLevelEditorSecondaryModeToolbar();
+	static FName GetSecondaryModeToolbarName();
+
 protected:
 
 	/**
@@ -90,4 +98,7 @@ private:
 	static FText GetActiveModeName(TWeakPtr<SLevelEditor> LevelEditorPtr);
 	static const FSlateBrush* GetActiveModeIcon(TWeakPtr<SLevelEditor> LevelEditorPtr);
 
+private:
+	static FName SecondaryModeToolbarName;
+
 };
diff --git a/Engine/Source/Editor/LevelEditor/Private/SLevelEditor.cpp b/Engine/Source/Editor/LevelEditor/Private/SLevelEditor.cpp
--- a/Engine/Source/Editor/LevelEditor/Private/SLevelEditor.cpp
+++ b/Engine/Source/Editor/LevelEditor/Private/SLevelEditor.cpp
@@ -181,6 +181,7 @@ void SLevelEditor::RegisterMenus()
 {
 	FLevelEditorMenu::RegisterLevelEditorMenus();
 	FLevelEditorToolBar::RegisterLevelEditorToolBar(LevelEditorCommands.ToSharedRef(), SharedThis(this));
+	FLevelEditorToolBar::RegisterLevelEditorSecondaryModeToolbar();
 }
 
 void SLevelEditor::Construct( const SLevelEditor::FArguments& InArgs)
@@ -304,6 +305,13 @@ void SLevelEditor::Initialize( const TSharedRef<SDockTab>& OwnerTab, const TShar
 	TSharedRef<SWidget> ContentArea = RestoreContentArea( OwnerTab, OwnerWindow );
 
 	FLevelEditorMenu::MakeLevelEditorMenu(LevelEditorCommands, SharedThis(this));
+
+	SecondaryModeToolbarWidget =
+		SNew(SBorder)
+		.Padding(0)
+		.BorderImage(FAppStyle::Get().GetBrush( "NoBorder" ));
+	
+	SecondaryModeToolbarWidget->SetContent(FLevelEditorToolBar::MakeLevelEditorSecondaryModeToolbar(LevelEditorCommands.ToSharedRef(), ModeUILayers));
 		
 	ChildSlot
 	[
@@ -315,6 +323,12 @@ void SLevelEditor::Initialize( const TSharedRef<SDockTab>& OwnerTab, const TShar
 			FLevelEditorToolBar::MakeLevelEditorToolBar(LevelEditorCommands.ToSharedRef(), SharedThis(this))
 		]
 		+SVerticalBox::Slot()
+		.Padding(FMargin(0.0f,0.0f,0.0f,2.0f))
+		.AutoHeight()
+		[
+			SecondaryModeToolbarWidget.ToSharedRef()
+		]
+		+SVerticalBox::Slot()
 		.Padding(4.0f, 2.f, 4.f, 2.f)
 		.FillHeight( 1.0f )
 		[
@@ -579,10 +593,13 @@ void SLevelEditor::OnToolkitHostingStarted( const TSharedRef< class IToolkit >&
 	//   Otherwise, it's going to be a huge cluster trying to distinguish tabs for different assets of the same type
 	//   of editor
 	TSharedPtr<FLevelEditorModeUILayer> ModeUILayer = MakeShareable(new FLevelEditorModeUILayer(this));
+	ModeUILayer->SetSecondaryModeToolbarName(FLevelEditorToolBar::GetSecondaryModeToolbarName());
 	HostedToolkits.Add( Toolkit );
 	ModeUILayer->OnToolkitHostingStarted(Toolkit);
 
 	ModeUILayers.Add(Toolkit->GetToolkitFName(), ModeUILayer);
+	
+	SecondaryModeToolbarWidget->SetContent(FLevelEditorToolBar::MakeLevelEditorSecondaryModeToolbar(LevelEditorCommands.ToSharedRef(), ModeUILayers));
 }
 
 void SLevelEditor::OnToolkitHostingFinished( const TSharedRef< class IToolkit >& Toolkit )
@@ -595,6 +612,8 @@ void SLevelEditor::OnToolkitHostingFinished( const TSharedRef< class IToolkit >&
 	}
 	HostedToolkits.Remove( Toolkit );
 
+	SecondaryModeToolbarWidget->SetContent(FLevelEditorToolBar::MakeLevelEditorSecondaryModeToolbar(LevelEditorCommands.ToSharedRef(), ModeUILayers));
+
 	// @todo toolkit minor: If user clicks X on all opened world-centric toolkit tabs, should we exit that toolkit automatically?
 	//   Feel 50/50 about this.  It's totally valid to use the "Save" menu even after closing tabs, etc.  Plus, you can spawn the tabs back up using the tab area down-down menu.
 }
diff --git a/Engine/Source/Editor/LevelEditor/Private/SLevelEditor.h b/Engine/Source/Editor/LevelEditor/Private/SLevelEditor.h
--- a/Engine/Source/Editor/LevelEditor/Private/SLevelEditor.h
+++ b/Engine/Source/Editor/LevelEditor/Private/SLevelEditor.h
@@ -372,6 +372,9 @@ private:
 
 	TMap<FName, TSharedPtr<FLevelEditorModeUILayer>> ModeUILayers;
 
+	/** Pointer to the widget that houses the level editor's mode toolbar */
+	TSharedPtr<SBorder> SecondaryModeToolbarWidget;
+	
 	/** Viewport context menu title, cached each time the element selection set changes */
 	FText CachedViewportContextMenuTitle;
 };
diff --git a/Engine/Source/Editor/LevelEditor/Private/SLevelEditorToolBox.cpp b/Engine/Source/Editor/LevelEditor/Private/SLevelEditorToolBox.cpp
--- a/Engine/Source/Editor/LevelEditor/Private/SLevelEditorToolBox.cpp
+++ b/Engine/Source/Editor/LevelEditor/Private/SLevelEditorToolBox.cpp
@@ -24,6 +24,7 @@
 #include "WorkspaceMenuStructure.h"
 #include "WorkspaceMenuStructureModule.h"
 #include "Misc/Attribute.h"
+#include "ToolMenus.h"
 
 #define LOCTEXT_NAMESPACE "SLevelEditorToolBox"
 
@@ -65,6 +66,13 @@ void FLevelEditorModeUILayer::OnToolkitHostingStarted(const TSharedRef<IToolkit>
 		Toolkit->RegisterTabSpawners(ToolkitHost->GetTabManager().ToSharedRef());
 		RegisterModeTabSpawners();
 		OnToolkitHostReadyForUI.ExecuteIfBound();
+
+		// Set up an owner for the current scope so that we can cleanly clean up the toolbar extension on hosting finish
+		FToolMenuOwnerScoped Owner(Toolkit->GetToolkitFName());
+
+		UToolMenu* SecondaryModeToolbar = UToolMenus::Get()->ExtendMenu(GetSecondaryModeToolbarName());
+
+		OnRegisterSecondaryModeToolbarExtension.ExecuteIfBound(SecondaryModeToolbar);
 	}
 
 }
@@ -74,6 +82,8 @@ void FLevelEditorModeUILayer::OnToolkitHostingFinished(const TSharedRef<IToolkit
 	if (HostedToolkit.IsValid() && HostedToolkit.Pin() == Toolkit)
 	{
 		FAssetEditorModeUILayer::OnToolkitHostingFinished(Toolkit);
+
+		UToolMenus::Get()->UnregisterOwnerByName(Toolkit->GetToolkitFName());
 	}
 }
 
diff --git a/Engine/Source/Editor/UnrealEd/Private/Toolkits/BaseToolkit.cpp b/Engine/Source/Editor/UnrealEd/Private/Toolkits/BaseToolkit.cpp
--- a/Engine/Source/Editor/UnrealEd/Private/Toolkits/BaseToolkit.cpp
+++ b/Engine/Source/Editor/UnrealEd/Private/Toolkits/BaseToolkit.cpp
@@ -223,7 +223,14 @@ FModeToolkit::~FModeToolkit()
 void FModeToolkit::SetModeUILayer(const TSharedPtr<FAssetEditorModeUILayer> InLayer)
 {
 	ModeUILayer = InLayer;
+
+	//TODO: Maybe Mode Toolbar Commands should be separate from the Toolkit Commands? Or the ModeUILayer should be
+	// responsible for a generic set of "Mode Commands"
+	ModeUILayer.Pin()->GetModeCommands()->Append(GetToolkitCommands());
 	ModeUILayer.Pin()->ToolkitHostReadyForUI().BindSP(this, &FModeToolkit::InvokeUI);
+
+	ModeUILayer.Pin()->RegisterSecondaryModeToolbarExtension().BindSP(this, &FModeToolkit::ExtendSecondaryModeToolbar);
+
 }
 
 void FModeToolkit::RegisterTabSpawners(const TSharedRef<FTabManager>& TabManager)
diff --git a/Engine/Source/Editor/UnrealEd/Public/Toolkits/BaseToolkit.h b/Engine/Source/Editor/UnrealEd/Public/Toolkits/BaseToolkit.h
--- a/Engine/Source/Editor/UnrealEd/Public/Toolkits/BaseToolkit.h
+++ b/Engine/Source/Editor/UnrealEd/Public/Toolkits/BaseToolkit.h
@@ -18,6 +18,7 @@ class IDetailsView;
 struct FDetailsViewArgs;
 class FAssetEditorModeUILayer;
 class SBorder;
+class UToolMenu;
 
 /**
  * Base class for all toolkits (abstract).
@@ -89,8 +90,6 @@ protected:
 	TSharedPtr<FWorkspaceItem> WorkspaceMenuCategory;
 };
 
-
-
 /**
  * This FModeToolkit just creates a basic UI panel that allows various InteractiveTools to
  * be initialized, and a DetailsView used to show properties of the active Tool.
@@ -163,6 +162,12 @@ public:
 	void SetModeSettingsObject(UObject* InSettingsObject);
 	virtual void InvokeUI();
 
+	/**
+	 * Override this function to extend the secondary mode toolbar (that appears below the main toolbar) for your mode by
+	 * filling in the UToolMenu that is passed in
+	 */
+	virtual void ExtendSecondaryModeToolbar(UToolMenu *InModeToolbarMenu) {}
+
 	/**
 	 * Force the Mode Toolbar/Palette to be repopulated with the current ToolPaletteNames
 	 */
diff --git a/Engine/Source/Runtime/SlateCore/Private/Styling/StarshipCoreStyle.cpp b/Engine/Source/Runtime/SlateCore/Private/Styling/StarshipCoreStyle.cpp
--- a/Engine/Source/Runtime/SlateCore/Private/Styling/StarshipCoreStyle.cpp
+++ b/Engine/Source/Runtime/SlateCore/Private/Styling/StarshipCoreStyle.cpp
@@ -2236,6 +2236,76 @@ void FStarshipCoreStyle::SetupMultiboxStyles(TSharedRef<FStyle>& Style)
 
 			Style->Set("CalloutToolbar", SlimToolbarStyle);
 		}
+
+		// Styling for secondary toolbars that appear below the main toolbar
+		{
+			FButtonStyle SecondaryToolBarButton = FButtonStyle(SimpleButtonStyle)
+			.SetNormalForeground(FStyleColors::Foreground)
+			.SetPressedForeground(FStyleColors::ForegroundHover)
+			.SetHoveredForeground(FStyleColors::ForegroundHover)
+			.SetDisabledForeground(FStyleColors::Foreground)
+			.SetNormalPadding(FMargin(8.f, 4.f, 8.f, 4.f))
+			.SetPressedPadding(FMargin(8.f, 5.f, 8.f, 3.f));
+
+			FComboButtonStyle SecondaryToolBarComboButton = FComboButtonStyle(Style->GetWidgetStyle<FComboButtonStyle>("SimpleComboButton"))
+				.SetDownArrowPadding(FMargin(1.f, 0.f, 2.f, 0.f));
+
+			
+			FToolBarStyle SecondaryToolbarStyle =
+				FToolBarStyle()
+				.SetShowLabels(false)
+				.SetBackground(*SlimToolbarBackground)
+				.SetBackgroundPadding(FMargin(4.f, 6.f, 0.f, 6.f))
+				.SetExpandBrush(IMAGE_BRUSH("Icons/toolbar_expand_16x", Icon16x16))
+				.SetComboButtonPadding(FMargin(0.0f, 0.0f))
+				.SetButtonPadding(FMargin(4.0f, 0.0f))
+				.SetCheckBoxPadding(FMargin(8.0f, 0.0f))
+				.SetSeparatorBrush(FSlateColorBrush(FStyleColors::Recessed))
+				.SetSeparatorPadding(FMargin(4.f, -5.0f))
+				.SetLabelStyle(FTextBlockStyle(NormalText))
+				.SetComboButtonStyle(SlimToolBarComboButton)
+				.SetLabelPadding(FMargin(4.f, 0.f, 0.f, 0.f))
+				.SetEditableTextStyle(FEditableTextBoxStyle(NormalEditableTextBoxStyle));
+
+			FCheckBoxStyle SecondaryToolBarToggleButtonCheckBoxStyle = FCheckBoxStyle(Style->GetWidgetStyle<FCheckBoxStyle>("ToggleButtonCheckbox"))
+				.SetPadding(FMargin(8.f, 4.f, 8.f, 4.f));
+
+			SecondaryToolbarStyle.SetToggleButtonStyle(SecondaryToolBarToggleButtonCheckBoxStyle);
+			SecondaryToolbarStyle.SetButtonStyle(SecondaryToolBarButton);
+
+			FButtonStyle SecondaryToolBarButtonSettingsRegion = FButtonStyle(SimpleButtonStyle)
+				.SetPressed(FSlateNoResource())
+				.SetHovered(FSlateNoResource())
+				.SetNormal(FSlateNoResource())
+				.SetNormalForeground(FStyleColors::Foreground)
+				.SetPressedForeground(FStyleColors::ForegroundHover)
+				.SetHoveredForeground(FStyleColors::ForegroundHover)
+				.SetDisabledForeground(FStyleColors::Foreground)
+				.SetNormalPadding(FMargin(8.f, 4.f, 8.f, 4.f))
+				.SetPressedPadding(FMargin(8.f, 5.f, 8.f, 3.f));
+
+
+			FCheckBoxStyle SecondaryToolbarSettingsToggleButtonStyle = FCheckBoxStyle(Style->GetWidgetStyle<FCheckBoxStyle>("TransparentCheckBox"))
+				.SetPadding(FMargin(8.f, 4.f, 8.f, 4.f));
+
+			SecondaryToolbarStyle.SetSettingsButtonStyle(SecondaryToolBarButtonSettingsRegion);
+			SecondaryToolbarStyle.SetSettingsToggleButtonStyle(SecondaryToolbarSettingsToggleButtonStyle);
+			SecondaryToolbarStyle.SetIconSize(Icon20x20);
+
+			SecondaryToolBarButtonSettingsRegion
+				.SetNormalPadding(FMargin(4.f, 4.f, 4.f, 4.f))
+				.SetPressedPadding(FMargin(4.f, 5.f, 4.f, 3.f));
+
+			FComboButtonStyle SecondaryToolBarSettingsComboButton = FComboButtonStyle(SecondaryToolBarComboButton)
+				.SetContentPadding(0)
+				.SetButtonStyle(SecondaryToolBarButtonSettingsRegion)
+				.SetDownArrowImage(IMAGE_BRUSH_SVG("Starship/Common/ellipsis-vertical-narrow", FVector2D(6, 20)));
+
+			SecondaryToolbarStyle.SetSettingsComboButtonStyle(SecondaryToolBarSettingsComboButton);
+			
+			Style->Set("SecondaryToolbar", SecondaryToolbarStyle);
+		}
+
 	}
 
 	// MenuBar
