MeshModelingTools: Fix capitalization of PolyGroups in secondary toolbar.

diff --git a/Engine/Plugins/Editor/ModelingToolsEditorMode/Source/ModelingToolsEditorMode/Private/ModelingModeToolkit_Toolbars.cpp b/Engine/Plugins/Editor/ModelingToolsEditorMode/Source/ModelingToolsEditorMode/Private/ModelingModeToolkit_Toolbars.cpp
--- a/Engine/Plugins/Editor/ModelingToolsEditorMode/Source/ModelingToolsEditorMode/Private/ModelingModeToolkit_Toolbars.cpp
+++ b/Engine/Plugins/Editor/ModelingToolsEditorMode/Source/ModelingToolsEditorMode/Private/ModelingModeToolkit_Toolbars.cpp
@@ -442,17 +442,17 @@ void FModelingToolsEditorModeToolkit::ExtendSecondaryModeToolbar(UToolMenu *InMo
 
 	Section.AddSeparator(NAME_None);
 
-	TSharedRef<SWidget> PolygroupsTextWidget = SNew(SHorizontalBox)
+	TSharedRef<SWidget> PolyGroupsTextWidget = SNew(SHorizontalBox)
 		+ SHorizontalBox::Slot()
 		.AutoWidth()
 		.Padding(8.0f, 0.0f, 8.0f, 0.0f)
 		.VAlign(VAlign_Center)
 		[
 			SNew(STextBlock)
-			.Text(LOCTEXT("PolygroupsSelectionText", "Polygroups"))
+			.Text(LOCTEXT("PolyGroupsSelectionText", "PolyGroups"))
 		];
 
-	Section.AddEntry(FToolMenuEntry::InitWidget("PolygroupsTextWidget", PolygroupsTextWidget, FText()));
+	Section.AddEntry(FToolMenuEntry::InitWidget("PolyGroupsTextWidget", PolyGroupsTextWidget, FText()));
 	
 	Section.AddEntry(FToolMenuEntry::InitToolBarButton(Commands.MeshSelectionModeAction_GroupFaces));
 	Section.AddEntry(FToolMenuEntry::InitToolBarButton(Commands.MeshSelectionModeAction_GroupEdges));
diff --git a/Engine/Plugins/Editor/ModelingToolsEditorMode/Source/ModelingToolsEditorMode/Private/ModelingToolsManagerActions.cpp b/Engine/Plugins/Editor/ModelingToolsEditorMode/Source/ModelingToolsEditorMode/Private/ModelingToolsManagerActions.cpp
--- a/Engine/Plugins/Editor/ModelingToolsEditorMode/Source/ModelingToolsEditorMode/Private/ModelingToolsManagerActions.cpp
+++ b/Engine/Plugins/Editor/ModelingToolsEditorMode/Source/ModelingToolsEditorMode/Private/ModelingToolsManagerActions.cpp
@@ -137,7 +137,7 @@ void FModelingToolsManagerCommands::RegisterCommands()
 
 	// UVs
 	REGISTER_MODELING_TOOL_COMMAND(BeginGlobalUVGenerateTool, "AutoUV", "Automatically unwrap and pack UVs");
-	REGISTER_MODELING_TOOL_COMMAND(BeginGroupUVGenerateTool, "Unwrap", "Recompute UVs for existing UV islands or Polygroups");
+	REGISTER_MODELING_TOOL_COMMAND(BeginGroupUVGenerateTool, "Unwrap", "Recompute UVs for existing UV islands or PolyGroups");
 	REGISTER_MODELING_TOOL_COMMAND(BeginUVProjectionTool, "Project", "Compute UVs via projecting to simple shapes");
 	REGISTER_MODELING_TOOL_COMMAND(BeginUVSeamEditTool, "SeamEd", "Add UV seams");
 	REGISTER_MODELING_TOOL_COMMAND(BeginTransformUVIslandsTool, "XFormUV", "Transform UV islands in UV space");
@@ -189,9 +189,9 @@ void FModelingToolsManagerCommands::RegisterCommands()
 	REGISTER_MODELING_TOOL_COMMAND(MeshSelectionModeAction_MeshTriangles, "Tris", "Select Mesh Triangles");
 	REGISTER_MODELING_TOOL_COMMAND(MeshSelectionModeAction_MeshVertices, "Verts", "Select Mesh Vertices");
 	REGISTER_MODELING_TOOL_COMMAND(MeshSelectionModeAction_MeshEdges, "Edges", "Select Mesh Edges");
-	REGISTER_MODELING_TOOL_COMMAND(MeshSelectionModeAction_GroupFaces, "Groups", "Select Mesh Polygroups");
-	REGISTER_MODELING_TOOL_COMMAND(MeshSelectionModeAction_GroupCorners, "Corners", "Select Mesh Polygroup Corners/Vertices");
-	REGISTER_MODELING_TOOL_COMMAND(MeshSelectionModeAction_GroupEdges, "Borders", "Select Mesh Polygroup Borders/Edges");
+	REGISTER_MODELING_TOOL_COMMAND(MeshSelectionModeAction_GroupFaces, "Groups", "Select Mesh PolyGroups");
+	REGISTER_MODELING_TOOL_COMMAND(MeshSelectionModeAction_GroupCorners, "Corners", "Select Mesh PolyGroup Corners/Vertices");
+	REGISTER_MODELING_TOOL_COMMAND(MeshSelectionModeAction_GroupEdges, "Borders", "Select Mesh PolyGroup Borders/Edges");
 
 	REGISTER_MODELING_ACTION_COMMAND(BeginSelectionAction_Delete, "Delete", "Delete the current Mesh Selection");
 	REGISTER_MODELING_ACTION_COMMAND(BeginSelectionAction_Disconnect, "Discon", "Disconnect the current Mesh Selection");
