Modeling updates to use ToolkitBuilder

diff --git a/Engine/Plugins/Editor/ModelingToolsEditorMode/Source/ModelingToolsEditorMode/Private/ModelingToolsEditorModeStyle.cpp b/Engine/Plugins/Editor/ModelingToolsEditorMode/Source/ModelingToolsEditorMode/Private/ModelingToolsEditorModeStyle.cpp
--- a/Engine/Plugins/Editor/ModelingToolsEditorMode/Source/ModelingToolsEditorMode/Private/ModelingToolsEditorModeStyle.cpp
+++ b/Engine/Plugins/Editor/ModelingToolsEditorMode/Source/ModelingToolsEditorMode/Private/ModelingToolsEditorModeStyle.cpp
@@ -104,6 +104,20 @@ void FModelingToolsEditorModeStyle::Initialize()
 		StyleSet->Set("ModelingToolsManagerCommands.CompleteActiveTool", new IMAGE_PLUGIN_BRUSH("Icons/icon_ActiveTool_Accept_40x", Icon20x20));
 		StyleSet->Set("ModelingToolsManagerCommands.CompleteActiveTool.Small", new IMAGE_PLUGIN_BRUSH("Icons/icon_ActiveTool_Accept_40x", Icon20x20));
 
+		StyleSet->Set("ModelingToolsManagerCommands.LoadFavoritesTools", new IMAGE_BRUSH_SVG( "Icons/LoadFavoritesTools", Icon20x20 ) );
+		StyleSet->Set("ModelingToolsManagerCommands.LoadShapesTools", new IMAGE_BRUSH_SVG( "Icons/LoadShapesTools", Icon20x20 ) );
+		StyleSet->Set("ModelingToolsManagerCommands.LoadCreateTools", new IMAGE_BRUSH_SVG( "Icons/LoadCreateTools", Icon20x20 ) );
+		StyleSet->Set("ModelingToolsManagerCommands.LoadPolyTools", new IMAGE_BRUSH_SVG( "Icons/LoadPolyTools", Icon20x20 ) );
+		StyleSet->Set("ModelingToolsManagerCommands.LoadTriTools", new IMAGE_BRUSH_SVG( "Icons/LoadTriTools", Icon20x20 ) );
+		StyleSet->Set("ModelingToolsManagerCommands.LoadDeformTools", new IMAGE_BRUSH_SVG( "Icons/LoadDeformTools", Icon20x20 ) );
+		StyleSet->Set("ModelingToolsManagerCommands.LoadTransformTools", new IMAGE_BRUSH_SVG( "Icons/LoadTransformTools", Icon20x20 ) );
+		StyleSet->Set("ModelingToolsManagerCommands.LoadMeshOpsTools", new IMAGE_BRUSH_SVG( "Icons/LoadMeshOpsTools", Icon20x20 ) );
+		StyleSet->Set("ModelingToolsManagerCommands.LoadVoxOpsTools", new IMAGE_BRUSH_SVG( "Icons/LoadVoxOpsTools", Icon20x20 ) );
+		StyleSet->Set("ModelingToolsManagerCommands.LoadAttributesTools", new IMAGE_BRUSH_SVG( "Icons/LoadAttributesTools", Icon20x20 ) );
+		StyleSet->Set("ModelingToolsManagerCommands.LoadUVsTools", new IMAGE_BRUSH_SVG( "Icons/LoadUVsTools", Icon20x20 ) );
+		StyleSet->Set("ModelingToolsManagerCommands.LoadBakingTools", new IMAGE_BRUSH_SVG( "Icons/LoadBakingTools", Icon20x20 ) );
+		StyleSet->Set("ModelingToolsManagerCommands.LoadVolumeTools", new IMAGE_BRUSH_SVG( "Icons/LoadVolumeTools", Icon20x20 ) );
+		StyleSet->Set("ModelingToolsManagerCommands.LoadLodsTools", new IMAGE_BRUSH_SVG( "Icons/LoadLodsTools", Icon20x20 ) );
 
 		StyleSet->Set("ModelingToolsManagerCommands.BeginShapeSprayTool", 				new IMAGE_PLUGIN_BRUSH("Icons/ShapeSpray_40x",	Icon20x20));
 		StyleSet->Set("ModelingToolsManagerCommands.BeginShapeSprayTool.Small", 		new IMAGE_PLUGIN_BRUSH("Icons/ShapeSpray_40x",	Icon20x20));
diff --git a/Engine/Plugins/Editor/ModelingToolsEditorMode/Source/ModelingToolsEditorMode/Private/ModelingToolsEditorModeToolkit.cpp b/Engine/Plugins/Editor/ModelingToolsEditorMode/Source/ModelingToolsEditorMode/Private/ModelingToolsEditorModeToolkit.cpp
--- a/Engine/Plugins/Editor/ModelingToolsEditorMode/Source/ModelingToolsEditorMode/Private/ModelingToolsEditorModeToolkit.cpp
+++ b/Engine/Plugins/Editor/ModelingToolsEditorMode/Source/ModelingToolsEditorMode/Private/ModelingToolsEditorModeToolkit.cpp
@@ -48,6 +48,7 @@
 #include "Widgets/Input/SEditableTextBox.h"
 #include "Widgets/Input/STextComboBox.h"
 #include "ToolMenus.h"
+#include "ToolkitBuilder.h"
 
 
 #define LOCTEXT_NAMESPACE "FModelingToolsEditorModeToolkit"
@@ -145,20 +146,261 @@ TSharedPtr<SWidget> FModelingToolsEditorModeToolkit::GetInlineContent() const
 		];
 }
 
+void FModelingToolsEditorModeToolkit::RegisterPalettes()
+{
+	const FModelingToolsManagerCommands& Commands = FModelingToolsManagerCommands::Get();
+	const TSharedPtr<FUICommandList> CommandList = GetToolkitCommands();
+
+	UEdMode* ScriptableMode = GetScriptableEditorMode().Get();
+	UModelingToolsModeCustomizationSettings* UISettings = GetMutableDefault<UModelingToolsModeCustomizationSettings>();
+	
+	ToolkitBuilder = MakeShared<FToolkitBuilder>(
+				ScriptableMode->GetModeInfo().ToolbarCustomizationName, 
+				GetToolkitCommands());
+
+	
+	const TSharedPtr<FEditablePalette> EditablePalette =
+		MakeShareable(new FEditablePalette(
+			Commands.LoadFavoritesTools.ToSharedRef(),
+			UISettings->ToolFavorites,
+			Commands.AddToFavorites.ToSharedRef(),
+			Commands.RemoveFromFavorites.ToSharedRef()));
+	ToolkitBuilder->AddPalette(StaticCastSharedPtr<FEditablePalette>(EditablePalette));
+	
+	ToolkitBuilder->AddPalette(
+		MakeShareable( new FToolPalette(
+		Commands.LoadShapesTools.ToSharedRef(),
+		TArray(
+		            {
+			            Commands.BeginAddBoxPrimitiveTool,
+			            Commands.BeginAddSpherePrimitiveTool,
+			            Commands.BeginAddCylinderPrimitiveTool,
+			            Commands.BeginAddConePrimitiveTool,
+			            Commands.BeginAddTorusPrimitiveTool,
+			            Commands.BeginAddArrowPrimitiveTool,
+			            Commands.BeginAddRectanglePrimitiveTool,
+			            Commands.BeginAddDiscPrimitiveTool,
+		            	Commands.BeginAddStairsPrimitiveTool
+		            }
+	            ))));
+
+	ToolkitBuilder->AddPalette(
+		MakeShareable( new FToolPalette(
+		Commands.LoadCreateTools.ToSharedRef(),
+		TArray(
+					{
+			            Commands.BeginDrawPolygonTool,
+			            Commands.BeginDrawPolyPathTool,
+			            Commands.BeginDrawAndRevolveTool,
+			            Commands.BeginRevolveBoundaryTool,
+			            Commands.BeginCombineMeshesTool,
+			            Commands.BeginDuplicateMeshesTool,
+			            Commands.BeginPatternTool
+		            }
+	            ))));
+
+	ToolkitBuilder->AddPalette(
+		MakeShareable( new FToolPalette(
+		Commands.LoadTransformTools.ToSharedRef(),
+		TArray(
+					{
+			            Commands.BeginTransformMeshesTool,
+			            Commands.BeginAlignObjectsTool,
+			            Commands.BeginEditPivotTool,
+			            Commands.BeginAddPivotActorTool,
+			            Commands.BeginBakeTransformTool,
+			            Commands.BeginTransferMeshTool,
+			            Commands.BeginConvertMeshesTool,
+			            Commands.BeginSplitMeshesTool
+		            }
+	            ))));
+
+	ToolkitBuilder->AddPalette(
+		MakeShareable( new FToolPalette(
+		Commands.LoadDeformTools.ToSharedRef(),
+		TArray(
+					{
+			            Commands.BeginSculptMeshTool,
+			            Commands.BeginRemeshSculptMeshTool,
+			            Commands.BeginSmoothMeshTool,
+			            Commands.BeginOffsetMeshTool,
+			            Commands.BeginMeshSpaceDeformerTool,
+			            Commands.BeginLatticeDeformerTool,
+			            Commands.BeginDisplaceMeshTool
+		            }
+	            ))));
+
+	ToolkitBuilder->AddPalette(
+		MakeShareable( new FToolPalette(
+		Commands.LoadMeshOpsTools.ToSharedRef(),
+		TArray(
+					{
+			            Commands.BeginSimplifyMeshTool,
+			            Commands.BeginRemeshMeshTool,
+			            Commands.BeginWeldEdgesTool,
+			            Commands.BeginRemoveOccludedTrianglesTool,
+			            Commands.BeginSelfUnionTool,
+			            Commands.BeginProjectToTargetTool
+		            }
+	            ))));
+
+	ToolkitBuilder->AddPalette(
+		MakeShareable( new FToolPalette(
+		Commands.LoadLodsTools.ToSharedRef(),
+		TArray(
+					{
+			            Commands.BeginLODManagerTool,
+			            Commands.BeginGenerateStaticMeshLODAssetTool,
+			            Commands.BeginISMEditorTool
+		            }
+	            ))));
+
+	ToolkitBuilder->AddPalette(
+		MakeShareable( new FToolPalette(
+		Commands.LoadVoxOpsTools.ToSharedRef(),
+		TArray(
+					{
+			            Commands.BeginVoxelSolidifyTool,
+			            Commands.BeginVoxelBlendTool,
+			            Commands.BeginVoxelMorphologyTool,
+		            }
+	            ))));
+
+
+	ToolkitBuilder->AddPalette(
+		MakeShareable( new FToolPalette(
+		Commands.LoadTriTools.ToSharedRef(),
+		TArray(
+					{
+			            Commands.BeginMeshSelectionTool,
+			            Commands.BeginTriEditTool,
+			            Commands.BeginHoleFillTool,
+			            Commands.BeginMirrorTool,
+			            Commands.BeginPlaneCutTool,
+			            Commands.BeginPolygonCutTool,
+			            Commands.BeginMeshTrimTool
+		            }
+	            ))));
+
+	ToolkitBuilder->AddPalette(
+		MakeShareable( new FToolPalette(
+		Commands.LoadPolyTools.ToSharedRef(),
+		TArray(
+					{
+				Commands.BeginPolyEditTool,
+				Commands.BeginPolyDeformTool,
+				Commands.BeginCubeGridTool,
+				Commands.BeginMeshBooleanTool,
+				Commands.BeginCutMeshWithMeshTool,
+				Commands.BeginSubdividePolyTool
+			}
+		))));
+
+	ToolkitBuilder->AddPalette(
+		MakeShareable( new FToolPalette(
+		Commands.LoadAttributesTools.ToSharedRef(),
+		TArray(
+					{
+				Commands.BeginMeshInspectorTool,
+				Commands.BeginEditNormalsTool,
+				Commands.BeginEditTangentsTool,
+				Commands.BeginAttributeEditorTool,
+				Commands.BeginPolyGroupsTool,
+				Commands.BeginMeshGroupPaintTool,
+				Commands.BeginMeshAttributePaintTool,
+				Commands.BeginEditMeshMaterialsTool
+			}
+		))));
+
+		ToolkitBuilder->AddPalette(
+			MakeShareable( new FToolPalette(
+			Commands.LoadBakingTools.ToSharedRef(),
+			TArray(
+						{
+				Commands.BeginBakeMeshAttributeMapsTool,
+				Commands.BeginBakeMultiMeshAttributeMapsTool,
+				Commands.BeginBakeMeshAttributeVertexTool,
+				Commands.BeginBakeRenderCaptureTool
+			}
+		))));
+
+
+
+	IGeometryProcessingInterfacesModule& GeomProcInterfaces = FModuleManager::Get().LoadModuleChecked<
+		IGeometryProcessingInterfacesModule>("GeometryProcessingInterfaces");
+	IGeometryProcessing_UVEditorAssetEditor* UVEditorAPI = GeomProcInterfaces.GetUVEditorAssetEditorImplementation();
+
+	TArray<TSharedPtr<FUICommandInfo>>&& PaletteActions =
+	{
+		Commands.BeginGlobalUVGenerateTool,
+		Commands.BeginGroupUVGenerateTool,
+		Commands.BeginUVProjectionTool,
+		Commands.BeginUVSeamEditTool,
+		Commands.BeginTransformUVIslandsTool,
+		Commands.BeginUVLayoutTool
+	};
+
+	if (UVEditorAPI)
+	{
+		PaletteActions.Add(Commands.LaunchUVEditor);
+	}
+	ToolkitBuilder->AddPalette(
+		MakeShareable( new FToolPalette(
+		Commands.LoadUVsTools.ToSharedRef(), PaletteActions)));
+
+	PaletteActions = {
+		Commands.BeginVolumeToMeshTool,
+		Commands.BeginMeshToVolumeTool
+	};
+
+	if (Commands.BeginBspConversionTool)
+	{
+		PaletteActions.Add(Commands.BeginBspConversionTool);
+	}
+
+	PaletteActions.Add(Commands.BeginPhysicsInspectorTool);
+	PaletteActions.Add(Commands.BeginSetCollisionGeometryTool);
+	PaletteActions.Add(Commands.BeginExtractCollisionGeometryTool);
+
+	ToolkitBuilder->OnToolEnded.BindLambda([this]()
+	{
+		GetScriptableEditorMode()->GetToolManager(EToolsContextScope::EdMode)->DeactivateTool(EToolSide::Left, EToolShutdownType::Cancel);
+	});
+
+	ToolkitBuilder->UpdateWidget();
+}
 
 void FModelingToolsEditorModeToolkit::Init(const TSharedPtr<IToolkitHost>& InitToolkitHost, TWeakObjectPtr<UEdMode> InOwningMode)
 {
+	bUsesToolkitBuilder = false;
+	
 	// Have to create the ToolkitWidget here because FModeToolkit::Init() is going to ask for it and add
 	// it to the Mode panel, and not ask again afterwards. However we have to call Init() to get the 
 	// ModeDetailsView created, that we need to add to the ToolkitWidget. So, we will create the Widget
 	// here but only add the rows to it after we call Init()
+
+	TSharedPtr<SHorizontalBox> ToolkitWidgetHBox;
 	TSharedPtr<SVerticalBox> ToolkitWidgetVBox = SNew(SVerticalBox);
-	SAssignNew(ToolkitWidget, SBorder)
-		.HAlign(HAlign_Fill)
-		.Padding(4)
-		[
-			ToolkitWidgetVBox->AsShared()
-		];
+	
+	if ( bUsesToolkitBuilder )
+	{
+		ToolkitWidgetHBox = SNew(SHorizontalBox);
+		SAssignNew(ToolkitWidget, SBorder)
+			.HAlign(HAlign_Fill)
+			.Padding(4)
+			[
+				ToolkitWidgetHBox->AsShared()
+			];		
+	}
+	else
+	{
+		SAssignNew(ToolkitWidget, SBorder)
+			.HAlign(HAlign_Fill)
+			.Padding(4)
+			[
+				ToolkitWidgetVBox->AsShared()
+			];
+	}
 
 	FModeToolkit::Init(InitToolkitHost, InOwningMode);
 
@@ -193,18 +435,32 @@ void FModelingToolsEditorModeToolkit::Init(const TSharedPtr<IToolkitHost>& InitT
 		[
 			ModeWarningArea->AsShared()
 		];
+
+	if (bUsesToolkitBuilder)
+	{
+		RegisterPalettes();
+
+		ToolkitWidgetVBox->AddSlot().AutoHeight().HAlign(HAlign_Fill).Padding(5)
+			[
+				ToolkitBuilder->GetToolPaletteWidget()
+			];
+	}
+	else
+	{
+		ToolkitWidgetVBox->AddSlot().AutoHeight().HAlign(HAlign_Fill).Padding(5)
+			[
+				ModeHeaderArea->AsShared()
+			];		
+	}
+		
 	ToolkitWidgetVBox->AddSlot().AutoHeight().HAlign(HAlign_Fill).Padding(5)
-		[
-			ModeHeaderArea->AsShared()
-		];
-	ToolkitWidgetVBox->AddSlot().AutoHeight().HAlign(HAlign_Fill).Padding(5)
-		[
-			ToolWarningArea->AsShared()
-		];
+	[
+		ToolWarningArea->AsShared()
+	];	
 	ToolkitWidgetVBox->AddSlot().HAlign(HAlign_Fill).FillHeight(1.f)
-		[
-			ModeDetailsView->AsShared()
-		];
+	[
+		ModeDetailsView->AsShared()
+	];
 	ToolkitWidgetVBox->AddSlot().AutoHeight().HAlign(HAlign_Fill).VAlign(VAlign_Bottom).Padding(5)
 		[
 			MakeAssetConfigPanel()->AsShared()
@@ -213,6 +469,26 @@ void FModelingToolsEditorModeToolkit::Init(const TSharedPtr<IToolkitHost>& InitT
 	ClearNotification();
 	ClearWarning();
 
+	if (bUsesToolkitBuilder)
+	{
+		ToolkitWidgetHBox->AddSlot()
+			.FillWidth(0.5f)
+			.MaxWidth(48)
+			.VAlign(VAlign_Fill)
+			.Padding(0)
+			[
+			ToolkitBuilder->CreateToolbarWidget()
+			];
+	
+		ToolkitWidgetHBox->AddSlot()
+			.FillWidth(4.5f)
+			.VAlign(VAlign_Fill)
+			.Padding(5)
+			[
+				ToolkitWidgetVBox->AsShared()
+			];
+	}
+
 	ActiveToolName = FText::GetEmpty();
 	ActiveToolMessage = FText::GetEmpty();
 
@@ -1083,6 +1359,12 @@ FText FModelingToolsEditorModeToolkit::GetToolPaletteDisplayName(FName Palette)
 
 void FModelingToolsEditorModeToolkit::BuildToolPalette(FName PaletteIndex, class FToolBarBuilder& ToolbarBuilder) 
 {
+
+	if (HasToolkitBuilder())
+	{
+		return;
+	}
+	
 	const FModelingToolsManagerCommands& Commands = FModelingToolsManagerCommands::Get();
 	UModelingToolsModeCustomizationSettings* UISettings = GetMutableDefault<UModelingToolsModeCustomizationSettings>();
 
@@ -1317,6 +1599,12 @@ void FModelingToolsEditorModeToolkit::InvokeUI()
 	// of FModeToolkit::InlineContentHolder so we can just replace it here
 	InlineContentHolder->SetContent(GetInlineContent().ToSharedRef());
 
+	// if the ToolkitBuilder is being used, it will make up the UI
+	if (HasToolkitBuilder())
+	{
+		return;
+	}
+
 	//
 	// Apply custom section header colors.
 	// See comments below, this is done via directly manipulating Slate widgets generated deep inside BaseToolkit.cpp,
diff --git a/Engine/Plugins/Editor/ModelingToolsEditorMode/Source/ModelingToolsEditorMode/Private/ModelingToolsManagerActions.cpp b/Engine/Plugins/Editor/ModelingToolsEditorMode/Source/ModelingToolsEditorMode/Private/ModelingToolsManagerActions.cpp
--- a/Engine/Plugins/Editor/ModelingToolsEditorMode/Source/ModelingToolsEditorMode/Private/ModelingToolsManagerActions.cpp
+++ b/Engine/Plugins/Editor/ModelingToolsEditorMode/Source/ModelingToolsEditorMode/Private/ModelingToolsManagerActions.cpp
@@ -209,6 +209,23 @@ void FModelingToolsManagerCommands::RegisterCommands()
 
 
 
+	UI_COMMAND(AddToFavorites, "Add to Favorites", "Add to Favorites", EUserInterfaceActionType::Button, FInputChord());
+	UI_COMMAND(RemoveFromFavorites, "Remove from Favorites", "Remove from Favorites", EUserInterfaceActionType::Button, FInputChord());
+	UI_COMMAND(LoadFavoritesTools, "Favorites", "Favorites", EUserInterfaceActionType::RadioButton, FInputChord());
+	UI_COMMAND(LoadShapesTools, "Shapes", "Shapes", EUserInterfaceActionType::RadioButton, FInputChord());
+	UI_COMMAND(LoadCreateTools, "Create", "Create", EUserInterfaceActionType::RadioButton, FInputChord());
+	UI_COMMAND(LoadPolyTools, "PolyEdit", "PolyEdit", EUserInterfaceActionType::RadioButton, FInputChord());
+	UI_COMMAND(LoadTriTools, "TriTools", "TriTools", EUserInterfaceActionType::RadioButton, FInputChord());
+	UI_COMMAND(LoadDeformTools, "Deform", "Deform", EUserInterfaceActionType::RadioButton, FInputChord());
+	UI_COMMAND(LoadTransformTools, "Transform", "Transform", EUserInterfaceActionType::RadioButton, FInputChord());
+	UI_COMMAND(LoadMeshOpsTools, "MeshOps", "MeshOps", EUserInterfaceActionType::RadioButton, FInputChord());
+	UI_COMMAND(LoadVoxOpsTools, "VoxOps", "VoxOps", EUserInterfaceActionType::RadioButton, FInputChord());
+	UI_COMMAND(LoadAttributesTools, "Attributes", "Atrribute", EUserInterfaceActionType::RadioButton, FInputChord());
+	UI_COMMAND(LoadUVsTools, "UVs", "UVs", EUserInterfaceActionType::RadioButton, FInputChord());
+	UI_COMMAND(LoadBakingTools, "Baking", "Baking", EUserInterfaceActionType::RadioButton, FInputChord());
+	UI_COMMAND(LoadVolumeTools, "Volumes", "Volumes", EUserInterfaceActionType::RadioButton, FInputChord());
+	UI_COMMAND(LoadLodsTools, "LODs", "LODs", EUserInterfaceActionType::RadioButton, FInputChord());
+
 	UI_COMMAND(AcceptActiveTool, "Accept", "Accept the active Tool", EUserInterfaceActionType::Button, FInputChord());
 	UI_COMMAND(CancelActiveTool, "Cancel", "Cancel the active Tool", EUserInterfaceActionType::Button, FInputChord());
 	UI_COMMAND(CompleteActiveTool, "Done", "Complete the active Tool", EUserInterfaceActionType::Button, FInputChord());
diff --git a/Engine/Plugins/Editor/ModelingToolsEditorMode/Source/ModelingToolsEditorMode/Public/ModelingToolsEditorModeToolkit.h b/Engine/Plugins/Editor/ModelingToolsEditorMode/Source/ModelingToolsEditorMode/Public/ModelingToolsEditorModeToolkit.h
--- a/Engine/Plugins/Editor/ModelingToolsEditorMode/Source/ModelingToolsEditorMode/Public/ModelingToolsEditorModeToolkit.h
+++ b/Engine/Plugins/Editor/ModelingToolsEditorMode/Source/ModelingToolsEditorMode/Public/ModelingToolsEditorModeToolkit.h
@@ -119,6 +119,9 @@ private:
 	void UpdateActiveToolProperties();
 	void InvalidateCachedDetailPanelState(UObject* ChangedObject);
 
+	/** A utility function to register the tool palettes with the ToolkitBuilder */
+	void RegisterPalettes();
+
 	TSharedPtr<SWidget> ToolShutdownViewportOverlayWidget;
 	void MakeToolShutdownOverlayWidget();
 	
diff --git a/Engine/Plugins/Editor/ModelingToolsEditorMode/Source/ModelingToolsEditorMode/Public/ModelingToolsManagerActions.h b/Engine/Plugins/Editor/ModelingToolsEditorMode/Source/ModelingToolsEditorMode/Public/ModelingToolsManagerActions.h
--- a/Engine/Plugins/Editor/ModelingToolsEditorMode/Source/ModelingToolsEditorMode/Public/ModelingToolsManagerActions.h
+++ b/Engine/Plugins/Editor/ModelingToolsEditorMode/Source/ModelingToolsEditorMode/Public/ModelingToolsManagerActions.h
@@ -166,6 +166,22 @@ public:
 	TSharedPtr<FUICommandInfo> BeginSelectionAction_Extrude;
 	TSharedPtr<FUICommandInfo> BeginSelectionAction_Offset;
 
+	TSharedPtr<FUICommandInfo> AddToFavorites;
+	TSharedPtr<FUICommandInfo> RemoveFromFavorites;
+	TSharedPtr<FUICommandInfo> LoadFavoritesTools;
+	TSharedPtr<FUICommandInfo> LoadShapesTools;
+	TSharedPtr<FUICommandInfo> LoadCreateTools;
+	TSharedPtr<FUICommandInfo> LoadPolyTools;
+	TSharedPtr<FUICommandInfo> LoadTriTools;
+	TSharedPtr<FUICommandInfo> LoadDeformTools;
+	TSharedPtr<FUICommandInfo> LoadTransformTools;
+	TSharedPtr<FUICommandInfo> LoadMeshOpsTools;
+	TSharedPtr<FUICommandInfo> LoadVoxOpsTools;
+	TSharedPtr<FUICommandInfo> LoadAttributesTools;
+	TSharedPtr<FUICommandInfo> LoadUVsTools;
+	TSharedPtr<FUICommandInfo> LoadBakingTools;
+	TSharedPtr<FUICommandInfo> LoadVolumeTools;
+	TSharedPtr<FUICommandInfo> LoadLodsTools;
 	// For connecting ModelingMode with UVEditor plugin
 	TSharedPtr<FUICommandInfo> LaunchUVEditor;
 
diff --git a/Engine/Source/Developer/WidgetRegistration/Private/WidgetRegistrationModule.cpp b/Engine/Source/Developer/WidgetRegistration/Private/WidgetRegistrationModule.cpp
--- a/Engine/Source/Developer/WidgetRegistration/Private/WidgetRegistrationModule.cpp
+++ b/Engine/Source/Developer/WidgetRegistration/Private/WidgetRegistrationModule.cpp
@@ -1,29 +1,17 @@
 // Copyright Epic Games, Inc. All Rights Reserved.
 
 // #include "ToolkitStyle.h"
-#include "Modules/ModuleInterface.h"
-#include "Modules/ModuleManager.h"
+#include "WidgetRegistrationModule.h"
+
 
 // @TODO: ~enable ToolkitStyle set 
 
-class FWidgetRegistrationModule
-	: public IModuleInterface
+void FWidgetRegistrationModule::StartupModule()
 {
-public:
-
-	// IModuleInterface interface
-	
-	virtual void StartupModule() override
-	{
-	//	FToolkitStyle::Initialize();
-	}
+}
 	
-	virtual void ShutdownModule() override
-	{
-	//	FToolkitStyle::Shutdown();
-	}
-
-};
+void FWidgetRegistrationModule::ShutdownModule()
+{
+}
 
 
-IMPLEMENT_MODULE(FWidgetRegistrationModule, WidgetRegistration);
diff --git a/Engine/Source/Developer/WidgetRegistration/Public/WidgetRegistrationModule.h b/Engine/Source/Developer/WidgetRegistration/Public/WidgetRegistrationModule.h
new file mode 100644
--- /dev/null
+++ b/Engine/Source/Developer/WidgetRegistration/Public/WidgetRegistrationModule.h
@@ -0,0 +1,26 @@
+// Copyright Epic Games, Inc. All Rights Reserved.
+
+// #include "ToolkitStyle.h"
+
+#pragma once
+
+#include "Modules/ModuleInterface.h"
+#include "Modules/ModuleManager.h"
+
+// @TODO: ~enable ToolkitStyle set 
+
+class FWidgetRegistrationModule
+	: public IModuleInterface
+{
+public:
+
+	// IModuleInterface interface
+	
+	virtual void StartupModule() override;
+	
+	virtual void ShutdownModule() override;
+
+};
+
+
+IMPLEMENT_MODULE(FWidgetRegistrationModule, WidgetRegistration);
