Update to correct inclusion of defunct Geometry API

diff --git a/Engine/Plugins/Editor/ModelingToolsEditorMode/Source/ModelingToolsEditorMode/Private/ModelingToolsEditorModeToolkit.cpp b/Engine/Plugins/Editor/ModelingToolsEditorMode/Source/ModelingToolsEditorMode/Private/ModelingToolsEditorModeToolkit.cpp
--- a/Engine/Plugins/Editor/ModelingToolsEditorMode/Source/ModelingToolsEditorMode/Private/ModelingToolsEditorModeToolkit.cpp
+++ b/Engine/Plugins/Editor/ModelingToolsEditorMode/Source/ModelingToolsEditorMode/Private/ModelingToolsEditorModeToolkit.cpp
@@ -328,7 +328,6 @@ void FModelingToolsEditorModeToolkit::RegisterPalettes()
 
 	IGeometryProcessingInterfacesModule& GeomProcInterfaces = FModuleManager::Get().LoadModuleChecked<
 		IGeometryProcessingInterfacesModule>("GeometryProcessingInterfaces");
-	IGeometryProcessing_UVEditorAssetEditor* UVEditorAPI = GeomProcInterfaces.GetUVEditorAssetEditorImplementation();
 
 	TArray<TSharedPtr<FUICommandInfo>>&& PaletteActions =
 	{
@@ -340,10 +339,7 @@ void FModelingToolsEditorModeToolkit::RegisterPalettes()
 		Commands.BeginUVLayoutTool
 	};
 
-	if (UVEditorAPI)
-	{
-		PaletteActions.Add(Commands.LaunchUVEditor);
-	}
+
 	ToolkitBuilder->AddPalette(
 		MakeShareable( new FToolPalette(
 		Commands.LoadUVsTools.ToSharedRef(), PaletteActions)));
