MeshModelingTools: Fix slate throttling issue in non-debug builds that prevents gizmo numerical UI from working properly outside of a tool. Also fix behavior for typing values into the UI for the selection gizmo.

diff --git a/Engine/Plugins/Editor/ModelingToolsEditorMode/Source/ModelingToolsEditorMode/Private/ModelingSelectionInteraction.cpp b/Engine/Plugins/Editor/ModelingToolsEditorMode/Source/ModelingToolsEditorMode/Private/ModelingSelectionInteraction.cpp
--- a/Engine/Plugins/Editor/ModelingToolsEditorMode/Source/ModelingToolsEditorMode/Private/ModelingSelectionInteraction.cpp
+++ b/Engine/Plugins/Editor/ModelingToolsEditorMode/Source/ModelingToolsEditorMode/Private/ModelingSelectionInteraction.cpp
@@ -370,6 +370,14 @@ void UModelingSelectionInteraction::OnEndGizmoTransform(UTransformProxy* Proxy)
 {
 	if (bInActiveTransform)
 	{
+		// If we have an update we haven't applied, apply it now. This is particularly important if the
+		// user transformed the gizmo by typing a new value into the gizmo UI, in which case there will
+		// not have been a chance to do ApplyPendingTransformInteractions() between the update and end calls.
+		if (bGizmoUpdatePending)
+		{
+			ApplyPendingTransformInteractions();
+		}
+
 		SelectionManager->EndTransformation();
 		bInActiveTransform = false;
 
diff --git a/Engine/Plugins/Editor/ModelingToolsEditorMode/Source/ModelingToolsEditorMode/Private/ModelingToolsEditorMode.cpp b/Engine/Plugins/Editor/ModelingToolsEditorMode/Source/ModelingToolsEditorMode/Private/ModelingToolsEditorMode.cpp
--- a/Engine/Plugins/Editor/ModelingToolsEditorMode/Source/ModelingToolsEditorMode/Private/ModelingToolsEditorMode.cpp
+++ b/Engine/Plugins/Editor/ModelingToolsEditorMode/Source/ModelingToolsEditorMode/Private/ModelingToolsEditorMode.cpp
@@ -428,11 +428,20 @@ void UModelingToolsEditorMode::Enter()
 			[this](const FInputDeviceRay& DeviceRay) { return TestForEditorGizmoHit(DeviceRay); });
 		GetInteractiveToolsContext()->InputRouter->RegisterSource(SelectionInteraction);
 
-		// Disable the SnappingManager while the SelectionInteraction is editing a mesh via transform gizmo
-		SelectionInteraction->OnTransformBegin.AddLambda(
-			[this]() { SceneSnappingManager->PauseSceneGeometryUpdates(); });
-		SelectionInteraction->OnTransformEnd.AddLambda(
-			[this]() { SceneSnappingManager->UnPauseSceneGeometryUpdates(); });
+		SelectionInteraction->OnTransformBegin.AddLambda([this]() 
+		{ 
+			// Disable the SnappingManager while the SelectionInteraction is editing a mesh via transform gizmo
+			SceneSnappingManager->PauseSceneGeometryUpdates();
+
+			// If the transform is happening via the gizmo numerical UI, then we can run into the same slate
+			// throttling issues as tools. We need to continue receiving render/tick while user scrubs the slate values.
+			FSlateThrottleManager::Get().DisableThrottle(true);
+		});
+		SelectionInteraction->OnTransformEnd.AddLambda([this]() 
+		{ 
+			FSlateThrottleManager::Get().DisableThrottle(false);
+			SceneSnappingManager->UnPauseSceneGeometryUpdates(); 
+		});
 	}
 
 	// register level objects observer that will update the snapping manager as the scene changes
