ModelingMode: clean up new Modeling UI sections

Added UModelingToolsModeCustomizationSettings::bUseLegacyModelingPalette, when enabled, new UI/palette is used. Enabled by default. Toggling this option can be done in-editor but requires a restart to see short vs long tool names
When setting above is enabled, new "long" tool names are configured in FModelingToolsManagerCommands::RegisterCommands(). This is only done once per session so cannot be toggled interactively.
Updated FModelingToolsEditorModeToolkit::RegisterPalettes(). Now includes Selection Tools tab/palette, as well as Volumes (was missing before). Added various missing Tools.
Tweaked logic for showing/hiding tools, and enable/disable state, based on whether uising old or new UI.
Modeling Toolbar no longer shows icons that can pop up Settings dialogs in Restrictive mode.

diff --git a/Engine/Plugins/Editor/ModelingToolsEditorMode/Source/ModelingToolsEditorMode/Private/ModelingModeToolkit_Toolbars.cpp b/Engine/Plugins/Editor/ModelingToolsEditorMode/Source/ModelingToolsEditorMode/Private/ModelingModeToolkit_Toolbars.cpp
--- a/Engine/Plugins/Editor/ModelingToolsEditorMode/Source/ModelingToolsEditorMode/Private/ModelingModeToolkit_Toolbars.cpp
+++ b/Engine/Plugins/Editor/ModelingToolsEditorMode/Source/ModelingToolsEditorMode/Private/ModelingModeToolkit_Toolbars.cpp
@@ -379,13 +379,18 @@ TSharedRef<SWidget> MakeMenu_ModelingModeConfigSettings(FModelingToolsEditorMode
 		}));
 	MenuBuilder.EndSection();
 
-	MenuBuilder.BeginSection("Section_Settings", LOCTEXT("Section_Settings", "Quick Settings"));
-	MenuBuilder.AddSubMenu(
-		LOCTEXT("QuickSettingsSubMenu", "Jump To Settings"), LOCTEXT("QuickSettingsSubMenu_ToolTip", "Jump to sections of the Settings dialogs relevant to Modeling Mode"),
-		FNewMenuDelegate::CreateLambda([=](FMenuBuilder& SubMenuBuilder) {
-			MakeSubMenu_QuickSettings(SubMenuBuilder);
-			}));
-	MenuBuilder.EndSection();
+	// only show settings UI quick access in non-Restrictive mode
+	const UModelingToolsEditorModeSettings* Settings = GetDefault<UModelingToolsEditorModeSettings>();
+	if (!Settings->InRestrictiveMode())
+	{
+		MenuBuilder.BeginSection("Section_Settings", LOCTEXT("Section_Settings", "Quick Settings"));
+		MenuBuilder.AddSubMenu(
+			LOCTEXT("QuickSettingsSubMenu", "Jump To Settings"), LOCTEXT("QuickSettingsSubMenu_ToolTip", "Jump to sections of the Settings dialogs relevant to Modeling Mode"),
+			FNewMenuDelegate::CreateLambda([=](FMenuBuilder& SubMenuBuilder) {
+				MakeSubMenu_QuickSettings(SubMenuBuilder);
+				}));
+		MenuBuilder.EndSection();
+	}
 
 	TSharedRef<SWidget> MenuWidget = MenuBuilder.MakeWidget();
 	return MenuWidget;
diff --git a/Engine/Plugins/Editor/ModelingToolsEditorMode/Source/ModelingToolsEditorMode/Private/ModelingToolsEditorMode.cpp b/Engine/Plugins/Editor/ModelingToolsEditorMode/Source/ModelingToolsEditorMode/Private/ModelingToolsEditorMode.cpp
--- a/Engine/Plugins/Editor/ModelingToolsEditorMode/Source/ModelingToolsEditorMode/Private/ModelingToolsEditorMode.cpp
+++ b/Engine/Plugins/Editor/ModelingToolsEditorMode/Source/ModelingToolsEditorMode/Private/ModelingToolsEditorMode.cpp
@@ -357,6 +357,7 @@ void UModelingToolsEditorMode::Enter()
 	UEdMode::Enter();
 
 	const UModelingToolsEditorModeSettings* ModelingModeSettings = GetDefault<UModelingToolsEditorModeSettings>();
+	const UModelingToolsModeCustomizationSettings* ModelingEditorSettings = GetDefault<UModelingToolsModeCustomizationSettings>();
 
 	// Register builders for tool targets that the mode uses.
 	GetInteractiveToolsContext()->TargetManager->AddTargetFactory(NewObject<UStaticMeshComponentToolTargetFactory>(GetToolManager()));
@@ -737,27 +738,31 @@ void UModelingToolsEditorMode::Enter()
 	RegisterTool(ToolManagerCommands.BeginExtractCollisionGeometryTool, TEXT("BeginExtractCollisionGeometryTool"), ExtractCollisionGeoToolBuilder);
 
 
+
+	bool bAlwaysShowAllSelectionCommands = !ModelingEditorSettings->bUseLegacyModelingPalette;
+
 	// this function registers and tracks an active UGeometrySelectionEditCommand and it's associated UICommand
-	auto RegisterSelectionTool = [&](TSharedPtr<FUICommandInfo> UICommand, FString ToolIdentifier, UInteractiveToolBuilder* Builder, bool bRequiresActiveTarget, bool bRequiresSelection)
+	auto RegisterSelectionTool = [&, bAlwaysShowAllSelectionCommands](TSharedPtr<FUICommandInfo> UICommand, FString ToolIdentifier, UInteractiveToolBuilder* Builder, bool bRequiresActiveTarget, bool bRequiresSelection)
 	{
 		UEditorInteractiveToolsContext* UseToolsContext = GetInteractiveToolsContext(EToolsContextScope::EdMode);
 		const TSharedRef<FUICommandList>& CommandList = Toolkit->GetToolkitCommands();
 		UseToolsContext->ToolManager->RegisterToolType(ToolIdentifier, Builder);
 		CommandList->MapAction(UICommand,
 			FExecuteAction::CreateUObject(UseToolsContext, &UEdModeInteractiveToolsContext::StartTool, ToolIdentifier),
-			FCanExecuteAction::CreateWeakLambda(UseToolsContext, [this, ToolIdentifier, UseToolsContext, bRequiresActiveTarget, bRequiresSelection]() {
+			FCanExecuteAction::CreateWeakLambda(UseToolsContext, [this, ToolIdentifier, UseToolsContext, bRequiresActiveTarget, bRequiresSelection, bAlwaysShowAllSelectionCommands]() {
 				return ShouldToolStartBeAllowed(ToolIdentifier) &&
 					( GetSelectionManager()->HasActiveTargets() || bRequiresActiveTarget == false) &&
 					( GetSelectionManager()->HasSelection() || bRequiresSelection == false ) &&
-					( GetSelectionManager()->GetMeshTopologyMode() != UGeometrySelectionManager::EMeshTopologyMode::None ) &&
+					( GetSelectionManager()->GetMeshTopologyMode() != UGeometrySelectionManager::EMeshTopologyMode::None || bAlwaysShowAllSelectionCommands ) &&
 					UseToolsContext->ToolManager->CanActivateTool(EToolSide::Mouse, ToolIdentifier);
 			}),
 			FIsActionChecked::CreateUObject(UseToolsContext, &UEdModeInteractiveToolsContext::IsToolActive, EToolSide::Mouse, ToolIdentifier),
 			//FIsActionButtonVisible::CreateUObject(GetSelectionManager(), &UGeometrySelectionManager::HasSelection),
-			FIsActionButtonVisible::CreateWeakLambda(UseToolsContext, [this, bRequiresActiveTarget, bRequiresSelection]() {
-				return ( GetSelectionManager()->HasActiveTargets() || bRequiresActiveTarget == false) &&
+			FIsActionButtonVisible::CreateWeakLambda(UseToolsContext, [this, bRequiresActiveTarget, bRequiresSelection, bAlwaysShowAllSelectionCommands]() {
+				return (bAlwaysShowAllSelectionCommands) ? true : (
+					( GetSelectionManager()->HasActiveTargets() || bRequiresActiveTarget == false) &&
 					( GetSelectionManager()->HasSelection() || bRequiresSelection == false ) &&
-					( GetSelectionManager()->GetMeshTopologyMode() != UGeometrySelectionManager::EMeshTopologyMode::None );
+					( GetSelectionManager()->GetMeshTopologyMode() != UGeometrySelectionManager::EMeshTopologyMode::None ) );
 			}),
 			EUIActionRepeatMode::RepeatDisabled);
 	};
@@ -831,9 +836,9 @@ void UModelingToolsEditorMode::Enter()
 	};
 
 	// create and register InteractiveCommands for mesh selections
-	RegisterSelectionCommand(NewObject<UDeleteGeometrySelectionCommand>(), ToolManagerCommands.BeginSelectionAction_Delete, false);
-	RegisterSelectionCommand(NewObject<UDisconnectGeometrySelectionCommand>(), ToolManagerCommands.BeginSelectionAction_Disconnect, false);
-	RegisterSelectionCommand(NewObject<URetriangulateGeometrySelectionCommand>(), ToolManagerCommands.BeginSelectionAction_Retriangulate, false);
+	RegisterSelectionCommand(NewObject<UDeleteGeometrySelectionCommand>(), ToolManagerCommands.BeginSelectionAction_Delete, bAlwaysShowAllSelectionCommands);
+	RegisterSelectionCommand(NewObject<UDisconnectGeometrySelectionCommand>(), ToolManagerCommands.BeginSelectionAction_Disconnect, bAlwaysShowAllSelectionCommands);
+	RegisterSelectionCommand(NewObject<URetriangulateGeometrySelectionCommand>(), ToolManagerCommands.BeginSelectionAction_Retriangulate, bAlwaysShowAllSelectionCommands);
 	RegisterSelectionCommand(NewObject<UModifyGeometrySelectionCommand>(), ToolManagerCommands.BeginSelectionAction_SelectAll, true);
 	RegisterSelectionCommand(NewObject<UModifyGeometrySelectionCommand_Invert>(), ToolManagerCommands.BeginSelectionAction_Invert, true);
 	RegisterSelectionCommand(NewObject<UModifyGeometrySelectionCommand_ExpandToConnected>(), ToolManagerCommands.BeginSelectionAction_ExpandToConnected, true);
@@ -941,7 +946,6 @@ void UModelingToolsEditorMode::Enter()
 		[this](const UObject* Object)  { UpdateSelectionManagerOnEditorSelectionChange(); } );
 
 	// restore various settings
-	const UModelingToolsModeCustomizationSettings* ModelingEditorSettings = GetDefault<UModelingToolsModeCustomizationSettings>();
 	if (SelectionManager)
 	{
 		UE::Geometry::EGeometryElementType LastElementType = static_cast<UE::Geometry::EGeometryElementType>(ModelingEditorSettings->LastMeshSelectionElementType);
diff --git a/Engine/Plugins/Editor/ModelingToolsEditorMode/Source/ModelingToolsEditorMode/Private/ModelingToolsEditorModeStyle.cpp b/Engine/Plugins/Editor/ModelingToolsEditorMode/Source/ModelingToolsEditorMode/Private/ModelingToolsEditorModeStyle.cpp
--- a/Engine/Plugins/Editor/ModelingToolsEditorMode/Source/ModelingToolsEditorMode/Private/ModelingToolsEditorModeStyle.cpp
+++ b/Engine/Plugins/Editor/ModelingToolsEditorMode/Source/ModelingToolsEditorMode/Private/ModelingToolsEditorModeStyle.cpp
@@ -105,6 +105,7 @@ void FModelingToolsEditorModeStyle::Initialize()
 		StyleSet->Set("ModelingToolsManagerCommands.CompleteActiveTool.Small", new IMAGE_PLUGIN_BRUSH("Icons/icon_ActiveTool_Accept_40x", Icon20x20));
 
 		StyleSet->Set("ModelingToolsManagerCommands.LoadFavoritesTools", new IMAGE_BRUSH_SVG( "Icons/LoadFavoritesTools", Icon20x20 ) );
+		StyleSet->Set("ModelingToolsManagerCommands.LoadSelectionTools", new IMAGE_BRUSH_SVG("Icons/ModSelectionObject_16", Icon20x20 ) );
 		StyleSet->Set("ModelingToolsManagerCommands.LoadShapesTools", new IMAGE_BRUSH_SVG( "Icons/LoadShapesTools", Icon20x20 ) );
 		StyleSet->Set("ModelingToolsManagerCommands.LoadCreateTools", new IMAGE_BRUSH_SVG( "Icons/LoadCreateTools", Icon20x20 ) );
 		StyleSet->Set("ModelingToolsManagerCommands.LoadPolyTools", new IMAGE_BRUSH_SVG( "Icons/LoadPolyTools", Icon20x20 ) );
diff --git a/Engine/Plugins/Editor/ModelingToolsEditorMode/Source/ModelingToolsEditorMode/Private/ModelingToolsEditorModeToolkit.cpp b/Engine/Plugins/Editor/ModelingToolsEditorMode/Source/ModelingToolsEditorMode/Private/ModelingToolsEditorModeToolkit.cpp
--- a/Engine/Plugins/Editor/ModelingToolsEditorMode/Source/ModelingToolsEditorMode/Private/ModelingToolsEditorModeToolkit.cpp
+++ b/Engine/Plugins/Editor/ModelingToolsEditorMode/Source/ModelingToolsEditorMode/Private/ModelingToolsEditorModeToolkit.cpp
@@ -159,204 +159,219 @@ void FModelingToolsEditorModeToolkit::RegisterPalettes()
 				GetToolkitCommands());
 
 	
-	const TSharedPtr<FEditablePalette> EditablePalette =
+	const TSharedPtr<FEditablePalette> FavoritesPalette =
 		MakeShareable(new FEditablePalette(
 			Commands.LoadFavoritesTools.ToSharedRef(),
 			UISettings->ToolFavorites,
 			Commands.AddToFavorites.ToSharedRef(),
 			Commands.RemoveFromFavorites.ToSharedRef()));
-	ToolkitBuilder->AddPalette(StaticCastSharedPtr<FEditablePalette>(EditablePalette));
+	ToolkitBuilder->AddPalette(StaticCastSharedPtr<FEditablePalette>(FavoritesPalette));
 	
-	ToolkitBuilder->AddPalette(
-		MakeShareable( new FToolPalette(
-		Commands.LoadShapesTools.ToSharedRef(),
-		TArray(
-		            {
-			            Commands.BeginAddBoxPrimitiveTool,
-			            Commands.BeginAddSpherePrimitiveTool,
-			            Commands.BeginAddCylinderPrimitiveTool,
-			            Commands.BeginAddConePrimitiveTool,
-			            Commands.BeginAddTorusPrimitiveTool,
-			            Commands.BeginAddArrowPrimitiveTool,
-			            Commands.BeginAddRectanglePrimitiveTool,
-			            Commands.BeginAddDiscPrimitiveTool,
-		            	Commands.BeginAddStairsPrimitiveTool
-		            }
-	            ))));
 
+	TArray<TSharedPtr<FUICommandInfo>> ShapesPaletteItems({
+		Commands.BeginAddBoxPrimitiveTool,
+		Commands.BeginAddSpherePrimitiveTool,
+		Commands.BeginAddCylinderPrimitiveTool,
+		Commands.BeginAddConePrimitiveTool,
+		Commands.BeginAddTorusPrimitiveTool,
+		Commands.BeginAddArrowPrimitiveTool,
+		Commands.BeginAddRectanglePrimitiveTool,
+		Commands.BeginAddDiscPrimitiveTool,
+		Commands.BeginAddStairsPrimitiveTool
+	});
 	ToolkitBuilder->AddPalette(
-		MakeShareable( new FToolPalette(
-		Commands.LoadCreateTools.ToSharedRef(),
-		TArray(
-					{
-			            Commands.BeginDrawPolygonTool,
-			            Commands.BeginDrawPolyPathTool,
-			            Commands.BeginDrawAndRevolveTool,
-			            Commands.BeginRevolveBoundaryTool,
-			            Commands.BeginCombineMeshesTool,
-			            Commands.BeginDuplicateMeshesTool,
-			            Commands.BeginPatternTool
-		            }
-	            ))));
+		MakeShareable( new FToolPalette( Commands.LoadShapesTools.ToSharedRef(), ShapesPaletteItems ) ) );
 
-	ToolkitBuilder->AddPalette(
-		MakeShareable( new FToolPalette(
-		Commands.LoadTransformTools.ToSharedRef(),
-		TArray(
-					{
-			            Commands.BeginTransformMeshesTool,
-			            Commands.BeginAlignObjectsTool,
-			            Commands.BeginEditPivotTool,
-			            Commands.BeginAddPivotActorTool,
-			            Commands.BeginBakeTransformTool,
-			            Commands.BeginTransferMeshTool,
-			            Commands.BeginConvertMeshesTool,
-			            Commands.BeginSplitMeshesTool
-		            }
-	            ))));
 
+
+	TArray<TSharedPtr<FUICommandInfo>> CreatePaletteItems({
+		Commands.BeginDrawPolygonTool,
+		Commands.BeginDrawPolyPathTool,
+		Commands.BeginDrawAndRevolveTool,
+		Commands.BeginRevolveBoundaryTool,
+		Commands.BeginCombineMeshesTool,
+		Commands.BeginDuplicateMeshesTool,
+		Commands.BeginPatternTool
+	});
 	ToolkitBuilder->AddPalette(
-		MakeShareable( new FToolPalette(
-		Commands.LoadDeformTools.ToSharedRef(),
-		TArray(
-					{
-			            Commands.BeginSculptMeshTool,
-			            Commands.BeginRemeshSculptMeshTool,
-			            Commands.BeginSmoothMeshTool,
-			            Commands.BeginOffsetMeshTool,
-			            Commands.BeginMeshSpaceDeformerTool,
-			            Commands.BeginLatticeDeformerTool,
-			            Commands.BeginDisplaceMeshTool
-		            }
-	            ))));
+		MakeShareable( new FToolPalette( Commands.LoadCreateTools.ToSharedRef(), CreatePaletteItems ) ) );
+
 
+	TArray<TSharedPtr<FUICommandInfo>> SelectionPaletteItems({
+		Commands.BeginSelectionAction_Delete,
+		Commands.BeginSelectionAction_Extrude,
+		Commands.BeginSelectionAction_Offset,
+
+		Commands.BeginPolyModelTool_PushPull,
+		Commands.BeginPolyModelTool_Inset,
+		Commands.BeginPolyModelTool_Outset,
+		Commands.BeginPolyModelTool_CutFaces,
+		Commands.BeginPolyModelTool_Bevel,
+
+		Commands.BeginPolyModelTool_InsertEdgeLoop,
+		Commands.BeginSelectionAction_Retriangulate,
+
+		Commands.BeginPolyModelTool_PolyEd,
+		Commands.BeginPolyModelTool_TriSel
+	});
+	ToolkitBuilder->AddPalette( 
+		MakeShareable( new FToolPalette( Commands.LoadSelectionTools.ToSharedRef(), SelectionPaletteItems ) ) );
+
+
+	TArray<TSharedPtr<FUICommandInfo>> TransformPaletteItems({
+		Commands.BeginTransformMeshesTool,
+		Commands.BeginAlignObjectsTool,
+		Commands.BeginEditPivotTool,
+		Commands.BeginAddPivotActorTool,
+		Commands.BeginBakeTransformTool,
+		Commands.BeginTransferMeshTool,
+		Commands.BeginConvertMeshesTool,
+		Commands.BeginSplitMeshesTool
+	});
 	ToolkitBuilder->AddPalette(
-		MakeShareable( new FToolPalette(
-		Commands.LoadMeshOpsTools.ToSharedRef(),
-		TArray(
-					{
-			            Commands.BeginSimplifyMeshTool,
-			            Commands.BeginRemeshMeshTool,
-			            Commands.BeginWeldEdgesTool,
-			            Commands.BeginRemoveOccludedTrianglesTool,
-			            Commands.BeginSelfUnionTool,
-			            Commands.BeginProjectToTargetTool
-		            }
-	            ))));
+		MakeShareable( new FToolPalette( Commands.LoadTransformTools.ToSharedRef(), TransformPaletteItems ) ) );
+
+
 
+	TArray<TSharedPtr<FUICommandInfo>> DeformPaletteItems({
+		Commands.BeginSculptMeshTool,
+		Commands.BeginRemeshSculptMeshTool,
+		Commands.BeginSmoothMeshTool,
+		Commands.BeginOffsetMeshTool,
+		Commands.BeginMeshSpaceDeformerTool,
+		Commands.BeginLatticeDeformerTool,
+		Commands.BeginDisplaceMeshTool
+	});
 	ToolkitBuilder->AddPalette(
-		MakeShareable( new FToolPalette(
-		Commands.LoadLodsTools.ToSharedRef(),
-		TArray(
-					{
-			            Commands.BeginLODManagerTool,
-			            Commands.BeginGenerateStaticMeshLODAssetTool,
-			            Commands.BeginISMEditorTool
-		            }
-	            ))));
+		MakeShareable( new FToolPalette( Commands.LoadDeformTools.ToSharedRef(), DeformPaletteItems ) ) );
+	
+
 
+	TArray<TSharedPtr<FUICommandInfo>> PolyToolsPaletteItems({
+		Commands.BeginPolyEditTool,
+		Commands.BeginPolyDeformTool,
+		Commands.BeginCubeGridTool,
+		Commands.BeginMeshBooleanTool,
+		Commands.BeginCutMeshWithMeshTool,
+		Commands.BeginSubdividePolyTool
+	});
 	ToolkitBuilder->AddPalette(
-		MakeShareable( new FToolPalette(
-		Commands.LoadVoxOpsTools.ToSharedRef(),
-		TArray(
-					{
-			            Commands.BeginVoxelSolidifyTool,
-			            Commands.BeginVoxelBlendTool,
-			            Commands.BeginVoxelMorphologyTool,
-		            }
-	            ))));
+		MakeShareable( new FToolPalette( Commands.LoadPolyTools.ToSharedRef(), PolyToolsPaletteItems ) ) );
 
 
+
+	TArray<TSharedPtr<FUICommandInfo>> TriToolsPaletteItems({
+		Commands.BeginMeshSelectionTool,
+		Commands.BeginTriEditTool,
+		Commands.BeginHoleFillTool,
+		Commands.BeginMirrorTool,
+		Commands.BeginPlaneCutTool,
+		Commands.BeginPolygonCutTool,
+		Commands.BeginMeshTrimTool
+	});
 	ToolkitBuilder->AddPalette(
-		MakeShareable( new FToolPalette(
-		Commands.LoadTriTools.ToSharedRef(),
-		TArray(
-					{
-			            Commands.BeginMeshSelectionTool,
-			            Commands.BeginTriEditTool,
-			            Commands.BeginHoleFillTool,
-			            Commands.BeginMirrorTool,
-			            Commands.BeginPlaneCutTool,
-			            Commands.BeginPolygonCutTool,
-			            Commands.BeginMeshTrimTool
-		            }
-	            ))));
+		MakeShareable( new FToolPalette( Commands.LoadTriTools.ToSharedRef(), TriToolsPaletteItems ) ) );
+
 
+
+	TArray<TSharedPtr<FUICommandInfo>> MeshOpPaletteItems({
+		Commands.BeginSimplifyMeshTool,
+		Commands.BeginRemeshMeshTool,
+		Commands.BeginWeldEdgesTool,
+		Commands.BeginRemoveOccludedTrianglesTool,
+		Commands.BeginSelfUnionTool,
+		Commands.BeginProjectToTargetTool
+	});
 	ToolkitBuilder->AddPalette(
-		MakeShareable( new FToolPalette(
-		Commands.LoadPolyTools.ToSharedRef(),
-		TArray(
-					{
-				Commands.BeginPolyEditTool,
-				Commands.BeginPolyDeformTool,
-				Commands.BeginCubeGridTool,
-				Commands.BeginMeshBooleanTool,
-				Commands.BeginCutMeshWithMeshTool,
-				Commands.BeginSubdividePolyTool
-			}
-		))));
+		MakeShareable( new FToolPalette( Commands.LoadMeshOpsTools.ToSharedRef(), MeshOpPaletteItems ) ) );
+
+
 
+	TArray<TSharedPtr<FUICommandInfo>> VoxelPaletteItems({
+		Commands.BeginVoxelSolidifyTool,
+		Commands.BeginVoxelBlendTool,
+		Commands.BeginVoxelMorphologyTool
+	});
+#if WITH_PROXYLOD
+	VoxelPaletteItems.Add(Commands.BeginVoxelBooleanTool);
+	VoxelPaletteItems.Add(Commands.BeginVoxelMergeTool);
+#endif	// WITH_PROXYLOD
 	ToolkitBuilder->AddPalette(
-		MakeShareable( new FToolPalette(
-		Commands.LoadAttributesTools.ToSharedRef(),
-		TArray(
-					{
-				Commands.BeginMeshInspectorTool,
-				Commands.BeginEditNormalsTool,
-				Commands.BeginEditTangentsTool,
-				Commands.BeginAttributeEditorTool,
-				Commands.BeginPolyGroupsTool,
-				Commands.BeginMeshGroupPaintTool,
-				Commands.BeginMeshAttributePaintTool,
-				Commands.BeginEditMeshMaterialsTool
-			}
-		))));
+		MakeShareable( new FToolPalette( Commands.LoadVoxOpsTools.ToSharedRef(), VoxelPaletteItems ) ) );
+
 
-		ToolkitBuilder->AddPalette(
-			MakeShareable( new FToolPalette(
-			Commands.LoadBakingTools.ToSharedRef(),
-			TArray(
-						{
-				Commands.BeginBakeMeshAttributeMapsTool,
-				Commands.BeginBakeMultiMeshAttributeMapsTool,
-				Commands.BeginBakeMeshAttributeVertexTool,
-				Commands.BeginBakeRenderCaptureTool
-			}
-		))));
 
+	TArray<TSharedPtr<FUICommandInfo>> BakingPaletteItems({
+		Commands.BeginBakeMeshAttributeMapsTool,
+		Commands.BeginBakeMultiMeshAttributeMapsTool,
+		Commands.BeginBakeMeshAttributeVertexTool,
+		Commands.BeginBakeRenderCaptureTool
+	});
+	ToolkitBuilder->AddPalette(
+		MakeShareable( new FToolPalette( Commands.LoadBakingTools.ToSharedRef(), BakingPaletteItems ) ) );
 
 
-	IGeometryProcessingInterfacesModule& GeomProcInterfaces = FModuleManager::Get().LoadModuleChecked<
-		IGeometryProcessingInterfacesModule>("GeometryProcessingInterfaces");
 
-	TArray<TSharedPtr<FUICommandInfo>>&& PaletteActions =
-	{
+	TArray<TSharedPtr<FUICommandInfo>> UVPaletteItems({
 		Commands.BeginGlobalUVGenerateTool,
 		Commands.BeginGroupUVGenerateTool,
 		Commands.BeginUVProjectionTool,
 		Commands.BeginUVSeamEditTool,
 		Commands.BeginTransformUVIslandsTool,
 		Commands.BeginUVLayoutTool
-	};
+	});
+	IGeometryProcessingInterfacesModule& GeomProcInterfaces = FModuleManager::Get().LoadModuleChecked<
+		IGeometryProcessingInterfacesModule>("GeometryProcessingInterfaces");
+	IGeometryProcessing_UVEditorAssetEditor* UVEditorAPI = GeomProcInterfaces.GetUVEditorAssetEditorImplementation();
+	if (UVEditorAPI)
+	{
+		UVPaletteItems.Add(Commands.LaunchUVEditor);
+	}
+
+	ToolkitBuilder->AddPalette(
+		MakeShareable( new FToolPalette( Commands.LoadUVsTools.ToSharedRef(), UVPaletteItems)) );
 
 
+
+	TArray<TSharedPtr<FUICommandInfo>> AttributesToolsPaletteItems({
+		Commands.BeginMeshInspectorTool,
+		Commands.BeginEditNormalsTool,
+		Commands.BeginEditTangentsTool,
+		Commands.BeginAttributeEditorTool,
+		Commands.BeginPolyGroupsTool,
+		Commands.BeginMeshGroupPaintTool,
+		Commands.BeginMeshAttributePaintTool,
+		Commands.BeginEditMeshMaterialsTool
+	});
 	ToolkitBuilder->AddPalette(
-		MakeShareable( new FToolPalette(
-		Commands.LoadUVsTools.ToSharedRef(), PaletteActions)));
+		MakeShareable( new FToolPalette( Commands.LoadAttributesTools.ToSharedRef(), AttributesToolsPaletteItems ) ) );
+
+
 
-	PaletteActions = {
+	TArray<TSharedPtr<FUICommandInfo>> VolumePaletteItems({
+		Commands.BeginPhysicsInspectorTool,
+		Commands.BeginSetCollisionGeometryTool,
+		Commands.BeginExtractCollisionGeometryTool,
 		Commands.BeginVolumeToMeshTool,
 		Commands.BeginMeshToVolumeTool
-	};
-
+	});
 	if (Commands.BeginBspConversionTool)
 	{
-		PaletteActions.Add(Commands.BeginBspConversionTool);
+		VolumePaletteItems.Add(Commands.BeginBspConversionTool);
 	}
+	ToolkitBuilder->AddPalette(
+		MakeShareable( new FToolPalette( Commands.LoadVolumeTools.ToSharedRef(), VolumePaletteItems)) );
+
+
+
+	TArray<TSharedPtr<FUICommandInfo>> LODPaletteItems({
+		Commands.BeginLODManagerTool,
+		Commands.BeginGenerateStaticMeshLODAssetTool,
+		Commands.BeginISMEditorTool
+	});
+	ToolkitBuilder->AddPalette(
+		MakeShareable( new FToolPalette( Commands.LoadLodsTools.ToSharedRef(), LODPaletteItems)) );
 
-	PaletteActions.Add(Commands.BeginPhysicsInspectorTool);
-	PaletteActions.Add(Commands.BeginSetCollisionGeometryTool);
-	PaletteActions.Add(Commands.BeginExtractCollisionGeometryTool);
 
 	ToolkitBuilder->OnToolEnded.BindLambda([this]()
 	{
@@ -368,7 +383,8 @@ void FModelingToolsEditorModeToolkit::RegisterPalettes()
 
 void FModelingToolsEditorModeToolkit::Init(const TSharedPtr<IToolkitHost>& InitToolkitHost, TWeakObjectPtr<UEdMode> InOwningMode)
 {
-	bUsesToolkitBuilder = false;
+	const UModelingToolsModeCustomizationSettings* UISettings = GetDefault<UModelingToolsModeCustomizationSettings>();
+	bUsesToolkitBuilder = !UISettings->bUseLegacyModelingPalette;
 	
 	// Have to create the ToolkitWidget here because FModeToolkit::Init() is going to ask for it and add
 	// it to the Mode panel, and not ask again afterwards. However we have to call Init() to get the 
diff --git a/Engine/Plugins/Editor/ModelingToolsEditorMode/Source/ModelingToolsEditorMode/Private/ModelingToolsManagerActions.cpp b/Engine/Plugins/Editor/ModelingToolsEditorMode/Source/ModelingToolsEditorMode/Private/ModelingToolsManagerActions.cpp
--- a/Engine/Plugins/Editor/ModelingToolsEditorMode/Source/ModelingToolsEditorMode/Private/ModelingToolsManagerActions.cpp
+++ b/Engine/Plugins/Editor/ModelingToolsEditorMode/Source/ModelingToolsEditorMode/Private/ModelingToolsManagerActions.cpp
@@ -40,6 +40,8 @@ TSharedPtr<FUICommandInfo> FModelingToolsManagerCommands::FindToolByName(FString
 void FModelingToolsManagerCommands::RegisterCommands()
 {
 	const UModelingToolsEditorModeSettings* Settings = GetDefault<UModelingToolsEditorModeSettings>();
+	const UModelingToolsModeCustomizationSettings* UISettings = GetDefault<UModelingToolsModeCustomizationSettings>();
+	bool bUseShortNames = UISettings->bUseLegacyModelingPalette;
 
 	// this has to be done with a compile-time macro because UI_COMMAND expands to LOCTEXT macros
 #define REGISTER_MODELING_TOOL_COMMAND(ToolCommandInfo, ToolName, ToolTip) \
@@ -54,118 +56,264 @@ void FModelingToolsManagerCommands::RegisterCommands()
 	// Shapes
 	REGISTER_MODELING_TOOL_COMMAND(BeginAddBoxPrimitiveTool, "Box", "Create new box objects");
 	REGISTER_MODELING_TOOL_COMMAND(BeginAddSpherePrimitiveTool, "Sphere", "Create new sphere objects");
-	REGISTER_MODELING_TOOL_COMMAND(BeginAddCylinderPrimitiveTool, "Cyl", "Create new cylinder objects");
 	REGISTER_MODELING_TOOL_COMMAND(BeginAddConePrimitiveTool, "Cone", "Create new cone objects");
 	REGISTER_MODELING_TOOL_COMMAND(BeginAddTorusPrimitiveTool, "Torus", "Create new torus objects");
 	REGISTER_MODELING_TOOL_COMMAND(BeginAddArrowPrimitiveTool, "Arrow", "Create new arrow objects");
-	REGISTER_MODELING_TOOL_COMMAND(BeginAddRectanglePrimitiveTool, "Rect", "Create new rectangle objects");
 	REGISTER_MODELING_TOOL_COMMAND(BeginAddDiscPrimitiveTool, "Disc", "Create new disc objects");
 	REGISTER_MODELING_TOOL_COMMAND(BeginAddStairsPrimitiveTool, "Stairs", "Create new stairs objects");
+	if ( bUseShortNames )
+	{
+		REGISTER_MODELING_TOOL_COMMAND(BeginAddCylinderPrimitiveTool, "Cyl", "Create new cylinder objects");
+		REGISTER_MODELING_TOOL_COMMAND(BeginAddRectanglePrimitiveTool, "Rect", "Create new rectangle objects");
+	}
+	else
+	{
+		REGISTER_MODELING_TOOL_COMMAND(BeginAddCylinderPrimitiveTool, "Cylinder", "Create new cylinder objects");
+		REGISTER_MODELING_TOOL_COMMAND(BeginAddRectanglePrimitiveTool, "Rectangle", "Create new rectangle objects");
+	}
+
 
 	// Create
-	REGISTER_MODELING_TOOL_COMMAND(BeginDrawPolygonTool, "PolyExt", "Draw and extrude 2D Polygons to create new objects");
-	REGISTER_MODELING_TOOL_COMMAND(BeginDrawPolyPathTool, "PathExt", "Draw and extrude 2D Paths to create new objects");
-	REGISTER_MODELING_TOOL_COMMAND(BeginDrawAndRevolveTool, "PathRev", "Draw and revolve 2D Paths to create new objects");
-	REGISTER_MODELING_TOOL_COMMAND(BeginRevolveBoundaryTool, "BdryRev", "Revolve Mesh boundary loops to create new objects");
-	REGISTER_MODELING_TOOL_COMMAND(BeginCombineMeshesTool, "Merge", "Merge multiple Meshes to create new objects");
-	REGISTER_MODELING_TOOL_COMMAND(BeginDuplicateMeshesTool, "Dupe", "Duplicate single Meshes to create new objects");
+	if ( bUseShortNames )
+	{
+		REGISTER_MODELING_TOOL_COMMAND(BeginDrawPolygonTool, "PolyExt", "Draw and extrude 2D Polygons to create new objects");
+		REGISTER_MODELING_TOOL_COMMAND(BeginDrawPolyPathTool, "PathExt", "Draw and extrude 2D Paths to create new objects");
+		REGISTER_MODELING_TOOL_COMMAND(BeginDrawAndRevolveTool, "PathRev", "Draw and revolve 2D Paths to create new objects");
+		REGISTER_MODELING_TOOL_COMMAND(BeginRevolveBoundaryTool, "BdryRev", "Revolve Mesh boundary loops to create new objects");
+		REGISTER_MODELING_TOOL_COMMAND(BeginCombineMeshesTool, "Merge", "Merge multiple Meshes to create new objects");
+		REGISTER_MODELING_TOOL_COMMAND(BeginDuplicateMeshesTool, "Dupe", "Duplicate single Meshes to create new objects");
+	}
+	else
+	{
+		REGISTER_MODELING_TOOL_COMMAND(BeginDrawPolygonTool, "Extrude Polygon", "Draw and extrude 2D Polygons to create new objects");
+		REGISTER_MODELING_TOOL_COMMAND(BeginDrawPolyPathTool, "Extrude Path", "Draw and extrude 2D Paths to create new objects");
+		REGISTER_MODELING_TOOL_COMMAND(BeginDrawAndRevolveTool, "Revolve Path", "Draw and revolve 2D Paths to create new objects");
+		REGISTER_MODELING_TOOL_COMMAND(BeginRevolveBoundaryTool, "Revolve Boundary", "Revolve Mesh boundary loops to create new objects");
+		REGISTER_MODELING_TOOL_COMMAND(BeginCombineMeshesTool, "Merge", "Merge multiple Meshes to create new objects");
+		REGISTER_MODELING_TOOL_COMMAND(BeginDuplicateMeshesTool, "Duplicate", "Duplicate single Meshes to create new objects");
+	}
 
 	// PolyModel
-	REGISTER_MODELING_TOOL_COMMAND(BeginPolyEditTool, "PolyEd", "Edit Meshes via PolyGroups");
-	REGISTER_MODELING_TOOL_COMMAND(BeginPolyDeformTool, "PolyDef", "Deform Meshes via PolyGroups");
-	REGISTER_MODELING_TOOL_COMMAND(BeginCubeGridTool, "CubeGr", "Create block out Meshes using a repositionable grid");
-	REGISTER_MODELING_TOOL_COMMAND(BeginMeshBooleanTool, "MshBool", "Apply Boolean operations to the two selected Meshes");
-	REGISTER_MODELING_TOOL_COMMAND(BeginCutMeshWithMeshTool, "MshCut", "Split one Mesh into parts using a second Mesh");
-	REGISTER_MODELING_TOOL_COMMAND(BeginSubdividePolyTool, "SubDiv", "Subdivide the selected Mesh via PolyGroups or triangles");
+	if ( bUseShortNames )
+	{
+		REGISTER_MODELING_TOOL_COMMAND(BeginPolyEditTool, "PolyEd", "Edit Meshes via PolyGroups");
+		REGISTER_MODELING_TOOL_COMMAND(BeginPolyDeformTool, "PolyDef", "Deform Meshes via PolyGroups");
+		REGISTER_MODELING_TOOL_COMMAND(BeginCubeGridTool, "CubeGr", "Create block out Meshes using a repositionable grid");
+		REGISTER_MODELING_TOOL_COMMAND(BeginMeshBooleanTool, "MshBool", "Apply Boolean operations to the two selected Meshes");
+		REGISTER_MODELING_TOOL_COMMAND(BeginCutMeshWithMeshTool, "MshCut", "Split one Mesh into parts using a second Mesh");
+		REGISTER_MODELING_TOOL_COMMAND(BeginSubdividePolyTool, "SubDiv", "Subdivide the selected Mesh via PolyGroups or triangles");
+	}
+	else
+	{
+		REGISTER_MODELING_TOOL_COMMAND(BeginPolyEditTool, "PolyGroup Edit", "Edit Meshes via PolyGroups");
+		REGISTER_MODELING_TOOL_COMMAND(BeginPolyDeformTool, "Deform PolyGroups", "Deform Meshes via PolyGroups");
+		REGISTER_MODELING_TOOL_COMMAND(BeginCubeGridTool, "CubeGrid", "Create block out Meshes using a repositionable grid");
+		REGISTER_MODELING_TOOL_COMMAND(BeginMeshBooleanTool, "Boolean", "Apply Boolean operations to the two selected Meshes");
+		REGISTER_MODELING_TOOL_COMMAND(BeginCutMeshWithMeshTool, "Mesh Cut", "Split one Mesh into parts using a second Mesh");
+		REGISTER_MODELING_TOOL_COMMAND(BeginSubdividePolyTool, "Subdivide", "Subdivide the selected Mesh via PolyGroups or triangles");
+	}
 
 	// TriModel
-	REGISTER_MODELING_TOOL_COMMAND(BeginMeshSelectionTool, "TriSel", "Select and edit Mesh triangles with a brush interface");
-	REGISTER_MODELING_TOOL_COMMAND(BeginTriEditTool, "TriEd", "Select and Edit the Mesh vertices, edges, and triangles");
-	REGISTER_MODELING_TOOL_COMMAND(BeginHoleFillTool, "HFill", "Fill holes in the selected Mesh");
+	if ( bUseShortNames )
+	{
+		REGISTER_MODELING_TOOL_COMMAND(BeginMeshSelectionTool, "TriSel", "Select and edit Mesh triangles with a brush interface");
+		REGISTER_MODELING_TOOL_COMMAND(BeginTriEditTool, "TriEd", "Select and Edit the Mesh vertices, edges, and triangles");
+		REGISTER_MODELING_TOOL_COMMAND(BeginHoleFillTool, "HFill", "Fill holes in the selected Mesh");
+		REGISTER_MODELING_TOOL_COMMAND(BeginPlaneCutTool, "PlnCut", "Cut the selected Meshes with a 3D plane");
+	}
+	else
+	{
+		REGISTER_MODELING_TOOL_COMMAND(BeginMeshSelectionTool, "Tri Select", "Select and edit Mesh triangles with a brush interface");
+		REGISTER_MODELING_TOOL_COMMAND(BeginTriEditTool, "Triangle Edit", "Select and Edit the Mesh vertices, edges, and triangles");
+		REGISTER_MODELING_TOOL_COMMAND(BeginHoleFillTool, "Fill Holes", "Fill holes in the selected Mesh");
+		REGISTER_MODELING_TOOL_COMMAND(BeginPlaneCutTool, "Plane Cut", "Cut the selected Meshes with a 3D plane");
+	}
 	REGISTER_MODELING_TOOL_COMMAND(BeginMirrorTool, "Mirror", "Mirror the selected Meshes");
-	REGISTER_MODELING_TOOL_COMMAND(BeginPlaneCutTool, "PlnCut", "Cut the selected Meshes with a 3D plane");
 	REGISTER_MODELING_TOOL_COMMAND(BeginPolygonCutTool, "PolyCut", "Cut the selected Mesh with an extruded polygon");
 	REGISTER_MODELING_TOOL_COMMAND(BeginMeshTrimTool, "Trim", "Trim/Cut the selected Mesh with the second selected Mesh");
 
 	// Deform
-	REGISTER_MODELING_TOOL_COMMAND(BeginSculptMeshTool, "VSclpt", "Vertex sculpting");
-	REGISTER_MODELING_TOOL_COMMAND(BeginRemeshSculptMeshTool, "DSclpt", "Dynamic mesh sculpting");
+	if ( bUseShortNames )
+	{
+		REGISTER_MODELING_TOOL_COMMAND(BeginSculptMeshTool, "VSclpt", "Vertex sculpting");
+		REGISTER_MODELING_TOOL_COMMAND(BeginRemeshSculptMeshTool, "DSclpt", "Dynamic mesh sculpting");
+		REGISTER_MODELING_TOOL_COMMAND(BeginDisplaceMeshTool, "Displce", "Tessellate and Displace the selected Mesh");
+	}
+	else
+	{
+		REGISTER_MODELING_TOOL_COMMAND(BeginSculptMeshTool, "Vertex Sculpt", "Vertex sculpting");
+		REGISTER_MODELING_TOOL_COMMAND(BeginRemeshSculptMeshTool, "Dynamic Sculpt", "Dynamic mesh sculpting");
+		REGISTER_MODELING_TOOL_COMMAND(BeginDisplaceMeshTool, "Displace", "Tessellate and Displace the selected Mesh");
+	}
 	REGISTER_MODELING_TOOL_COMMAND(BeginSmoothMeshTool, "Smooth", "Smooth the shape of the selected Mesh");
 	REGISTER_MODELING_TOOL_COMMAND(BeginOffsetMeshTool, "Offset", "Offset the surface of the selected Mesh");
 	REGISTER_MODELING_TOOL_COMMAND(BeginMeshSpaceDeformerTool, "Warp", "Reshape the selected Mesh using space deformers");
 	REGISTER_MODELING_TOOL_COMMAND(BeginLatticeDeformerTool, "Lattice", "Deform the selected Mesh using a 3D lattice/grid");
-	REGISTER_MODELING_TOOL_COMMAND(BeginDisplaceMeshTool, "Displce", "Tessellate and Displace the selected Mesh");
 
 	// Transform
-	REGISTER_MODELING_TOOL_COMMAND(BeginTransformMeshesTool, "XForm", "Transform the selected Meshes");
+	if ( bUseShortNames )
+	{
+		REGISTER_MODELING_TOOL_COMMAND(BeginTransformMeshesTool, "XForm", "Transform the selected Meshes");
+		REGISTER_MODELING_TOOL_COMMAND(BeginEditPivotTool, "Pivot", "Edit the pivot points of the selected Meshes");
+		REGISTER_MODELING_TOOL_COMMAND(BeginAddPivotActorTool, "PivotAct", "Add actor to act as a pivot for child components");
+		REGISTER_MODELING_TOOL_COMMAND(BeginBakeTransformTool, "BakeRS", "Bake rotation and scale into the selected Meshes");
+	}
+	else
+	{
+		REGISTER_MODELING_TOOL_COMMAND(BeginTransformMeshesTool, "Transform", "Transform the selected Meshes");
+		REGISTER_MODELING_TOOL_COMMAND(BeginEditPivotTool, "Edit Pivot", "Edit the pivot points of the selected Meshes");
+		REGISTER_MODELING_TOOL_COMMAND(BeginAddPivotActorTool, "Pivot Actor", "Add actor to act as a pivot for child components");
+		REGISTER_MODELING_TOOL_COMMAND(BeginBakeTransformTool, "Bake Transform", "Bake rotation and scale into the selected Meshes");
+	}
 	REGISTER_MODELING_TOOL_COMMAND(BeginAlignObjectsTool, "Align", "Align the selected Objects");
-	REGISTER_MODELING_TOOL_COMMAND(BeginEditPivotTool, "Pivot", "Edit the pivot points of the selected Meshes");
-	REGISTER_MODELING_TOOL_COMMAND(BeginAddPivotActorTool, "PivotAct", "Add actor to act as a pivot for child components");
-	REGISTER_MODELING_TOOL_COMMAND(BeginBakeTransformTool, "BakeRS", "Bake rotation and scale into the selected Meshes");
 	REGISTER_MODELING_TOOL_COMMAND(BeginTransferMeshTool, "Transfer", "Copy the first selected Mesh to the second selected Mesh");
 	REGISTER_MODELING_TOOL_COMMAND(BeginConvertMeshesTool, "Convert", "Convert the selected Meshes to a different type of Mesh Object");
 	REGISTER_MODELING_TOOL_COMMAND(BeginSplitMeshesTool, "Split", "Split the selected Meshes into separate parts based on connectivity");
 	REGISTER_MODELING_TOOL_COMMAND(BeginPatternTool, "Pattern", "Create patterns of Meshes");
 
 	// MeshOps
-	REGISTER_MODELING_TOOL_COMMAND(BeginSimplifyMeshTool, "Simplfy", "Simplify the selected Mesh");
+	if ( bUseShortNames )
+	{
+		REGISTER_MODELING_TOOL_COMMAND(BeginSimplifyMeshTool, "Simplfy", "Simplify the selected Mesh");
+	}
+	else
+	{
+		REGISTER_MODELING_TOOL_COMMAND(BeginSimplifyMeshTool, "Simplify", "Simplify the selected Mesh");
+	}
 	REGISTER_MODELING_TOOL_COMMAND(BeginRemeshMeshTool, "Remesh", "Re-triangulate the selected Mesh");
 	REGISTER_MODELING_TOOL_COMMAND(BeginWeldEdgesTool, "Weld", "Weld overlapping Mesh edges");
 	REGISTER_MODELING_TOOL_COMMAND(BeginRemoveOccludedTrianglesTool, "Jacket", "Remove hidden triangles from the selected Meshes");
 	REGISTER_MODELING_TOOL_COMMAND(BeginSelfUnionTool, "Union", "Boolean Union the selected Meshes, including Self-Union to resolve self-intersections");
 	REGISTER_MODELING_TOOL_COMMAND(BeginProjectToTargetTool, "Project", "Map/re-mesh the first selected Mesh onto the second selected Mesh");
 
+
 	// VoxOps
-	REGISTER_MODELING_TOOL_COMMAND(BeginVoxelSolidifyTool, "VoxWrap", "Wrap the selected Meshes using voxels");
-	REGISTER_MODELING_TOOL_COMMAND(BeginVoxelBlendTool, "VoxBlnd", "Blend the selected Meshes using voxels");
-	REGISTER_MODELING_TOOL_COMMAND(BeginVoxelMorphologyTool, "VoxMrph", "Offset/Inset the selected Meshes using voxels");
+	if ( bUseShortNames )
+	{
+		REGISTER_MODELING_TOOL_COMMAND(BeginVoxelSolidifyTool, "VoxWrap", "Wrap the selected Meshes using voxels");
+		REGISTER_MODELING_TOOL_COMMAND(BeginVoxelBlendTool, "VoxBlnd", "Blend the selected Meshes using voxels");
+		REGISTER_MODELING_TOOL_COMMAND(BeginVoxelMorphologyTool, "VoxMrph", "Offset/Inset the selected Meshes using voxels");
 #if WITH_PROXYLOD
-	// The ProxyLOD plugin is currently only available on Windows. Without it, the following tools do not work as expected.
-	REGISTER_MODELING_TOOL_COMMAND(BeginVoxelBooleanTool, "VoxBool", "Boolean the selected Meshes using voxels");
-	REGISTER_MODELING_TOOL_COMMAND(BeginVoxelMergeTool, "VoxMrg", "Merge the selected Meshes using voxels");
+		// The ProxyLOD plugin is currently only available on Windows. Without it, the following tools do not work as expected.
+		REGISTER_MODELING_TOOL_COMMAND(BeginVoxelBooleanTool, "VoxBool", "Boolean the selected Meshes using voxels");
+		REGISTER_MODELING_TOOL_COMMAND(BeginVoxelMergeTool, "VoxMrg", "Merge the selected Meshes using voxels");
 #endif	// WITH_PROXYLOD
+	}
+	else
+	{
+		REGISTER_MODELING_TOOL_COMMAND(BeginVoxelSolidifyTool, "Vox Wrap", "Wrap the selected Meshes using voxels");
+		REGISTER_MODELING_TOOL_COMMAND(BeginVoxelBlendTool, "Vox Blend", "Blend the selected Meshes using voxels");
+		REGISTER_MODELING_TOOL_COMMAND(BeginVoxelMorphologyTool, "Vox Offset", "Offset/Inset the selected Meshes using voxels");
+#if WITH_PROXYLOD
+		// The ProxyLOD plugin is currently only available on Windows. Without it, the following tools do not work as expected.
+		REGISTER_MODELING_TOOL_COMMAND(BeginVoxelBooleanTool, "Vox Boolean", "Boolean the selected Meshes using voxels");
+		REGISTER_MODELING_TOOL_COMMAND(BeginVoxelMergeTool, "Vox Merge", "Merge the selected Meshes using voxels");
+#endif	// WITH_PROXYLOD
+	}
 
 	// Attributes
-	REGISTER_MODELING_TOOL_COMMAND(BeginMeshInspectorTool, "Inspct", "Inspect Mesh attributes");
-	REGISTER_MODELING_TOOL_COMMAND(BeginEditNormalsTool, "Nrmls", "Recompute or Repair Normals");
-	REGISTER_MODELING_TOOL_COMMAND(BeginEditTangentsTool, "Tngnts", "Recompute Tangents");
-	REGISTER_MODELING_TOOL_COMMAND(BeginAttributeEditorTool, "AttrEd", "Edit/configure Mesh attributes");
-	REGISTER_MODELING_TOOL_COMMAND(BeginPolyGroupsTool, "GrpGen", "Generate new PolyGroups");
-	REGISTER_MODELING_TOOL_COMMAND(BeginMeshGroupPaintTool, "GrpPnt", "Paint new PolyGroups");
-	REGISTER_MODELING_TOOL_COMMAND(BeginMeshAttributePaintTool, "MapPnt", "Paint attribute maps");
-	REGISTER_MODELING_TOOL_COMMAND(BeginEditMeshMaterialsTool, "MatEd", "Assign materials to selected triangles");
+	if ( bUseShortNames )
+	{
+		REGISTER_MODELING_TOOL_COMMAND(BeginMeshInspectorTool, "Inspct", "Inspect Mesh attributes");
+		REGISTER_MODELING_TOOL_COMMAND(BeginEditNormalsTool, "Nrmls", "Recompute or Repair Normals");
+		REGISTER_MODELING_TOOL_COMMAND(BeginEditTangentsTool, "Tngnts", "Recompute Tangents");
+		REGISTER_MODELING_TOOL_COMMAND(BeginAttributeEditorTool, "AttrEd", "Edit/configure Mesh attributes");
+		REGISTER_MODELING_TOOL_COMMAND(BeginPolyGroupsTool, "GrpGen", "Generate new PolyGroups");
+		REGISTER_MODELING_TOOL_COMMAND(BeginMeshGroupPaintTool, "GrpPnt", "Paint new PolyGroups");
+		REGISTER_MODELING_TOOL_COMMAND(BeginMeshAttributePaintTool, "MapPnt", "Paint attribute maps");
+		REGISTER_MODELING_TOOL_COMMAND(BeginEditMeshMaterialsTool, "MatEd", "Assign materials to selected triangles");
+	}
+	else
+	{
+		REGISTER_MODELING_TOOL_COMMAND(BeginMeshInspectorTool, "Inspect", "Inspect Mesh attributes");
+		REGISTER_MODELING_TOOL_COMMAND(BeginEditNormalsTool, "Normals", "Recompute or Repair Normals");
+		REGISTER_MODELING_TOOL_COMMAND(BeginEditTangentsTool, "Tangents", "Recompute Tangents");
+		REGISTER_MODELING_TOOL_COMMAND(BeginPolyGroupsTool, "Generate PolyGroups", "Generate new PolyGroups");
+		REGISTER_MODELING_TOOL_COMMAND(BeginMeshGroupPaintTool, "Paint PolyGroups", "Paint new PolyGroups");
+		REGISTER_MODELING_TOOL_COMMAND(BeginAttributeEditorTool, "Edit Attributes", "Edit/configure Mesh attributes");
+		REGISTER_MODELING_TOOL_COMMAND(BeginMeshAttributePaintTool, "Paint Maps", "Paint attribute maps");
+		REGISTER_MODELING_TOOL_COMMAND(BeginEditMeshMaterialsTool, "Edit Materials", "Assign materials to selected triangles");
+	}
 
 	// UVs
-	REGISTER_MODELING_TOOL_COMMAND(BeginGlobalUVGenerateTool, "AutoUV", "Automatically unwrap and pack UVs");
-	REGISTER_MODELING_TOOL_COMMAND(BeginGroupUVGenerateTool, "Unwrap", "Recompute UVs for existing UV islands or PolyGroups");
-	REGISTER_MODELING_TOOL_COMMAND(BeginUVProjectionTool, "Project", "Compute UVs via projecting to simple shapes");
-	REGISTER_MODELING_TOOL_COMMAND(BeginUVSeamEditTool, "SeamEd", "Add UV seams");
-	REGISTER_MODELING_TOOL_COMMAND(BeginTransformUVIslandsTool, "XFormUV", "Transform UV islands in UV space");
-	REGISTER_MODELING_TOOL_COMMAND(BeginUVLayoutTool, "Layout", "Transform and Repack existing UVs");
-	// This is done directly, not with the REGISTER_ macro, since we don't want it added to the tool list or use a toggle button
-	UI_COMMAND(LaunchUVEditor, "UVEditor", "Launch UV asset editor", EUserInterfaceActionType::Button, FInputChord());
+	if ( bUseShortNames )
+	{
+		REGISTER_MODELING_TOOL_COMMAND(BeginGlobalUVGenerateTool, "AutoUV", "Automatically unwrap and pack UVs");
+		REGISTER_MODELING_TOOL_COMMAND(BeginGroupUVGenerateTool, "Unwrap", "Recompute UVs for existing UV islands or PolyGroups");
+		REGISTER_MODELING_TOOL_COMMAND(BeginUVProjectionTool, "Project", "Compute UVs via projecting to simple shapes");
+		REGISTER_MODELING_TOOL_COMMAND(BeginUVSeamEditTool, "SeamEd", "Add UV seams");
+		REGISTER_MODELING_TOOL_COMMAND(BeginTransformUVIslandsTool, "XFormUV", "Transform UV islands in UV space");
+		REGISTER_MODELING_TOOL_COMMAND(BeginUVLayoutTool, "Layout", "Transform and Repack existing UVs");
+		// This is done directly, not with the REGISTER_ macro, since we don't want it added to the tool list or use a toggle button
+		UI_COMMAND(LaunchUVEditor, "UVEditor", "Launch UV asset editor", EUserInterfaceActionType::Button, FInputChord());
+	}
+	else
+	{
+		REGISTER_MODELING_TOOL_COMMAND(BeginGlobalUVGenerateTool, "AutoUV", "Automatically unwrap and pack UVs");
+		REGISTER_MODELING_TOOL_COMMAND(BeginGroupUVGenerateTool, "UV Unwrap", "Recompute UVs for existing UV islands or PolyGroups");
+		REGISTER_MODELING_TOOL_COMMAND(BeginUVProjectionTool, "Project UVs", "Compute UVs via projecting to simple shapes");
+		REGISTER_MODELING_TOOL_COMMAND(BeginUVSeamEditTool, "Edit UV Seams", "Add UV seams");
+		REGISTER_MODELING_TOOL_COMMAND(BeginTransformUVIslandsTool, "Transform UVs", "Transform UV islands in UV space");
+		REGISTER_MODELING_TOOL_COMMAND(BeginUVLayoutTool, "Layout UVs", "Transform and Repack existing UVs");
+		// This is done directly, not with the REGISTER_ macro, since we don't want it added to the tool list or use a toggle button
+		UI_COMMAND(LaunchUVEditor, "UV Editor", "Launch UV asset editor", EUserInterfaceActionType::Button, FInputChord());
+	}
 
 	// Baking
-	REGISTER_MODELING_TOOL_COMMAND(BeginBakeMeshAttributeMapsTool, "BakeTx", "Bake textures for a target Mesh");
-	REGISTER_MODELING_TOOL_COMMAND(BeginBakeMultiMeshAttributeMapsTool, "BakeAll", "Bake textures for a target Mesh from multiple source Meshes");
-	REGISTER_MODELING_TOOL_COMMAND(BeginBakeRenderCaptureTool, "BakeRC", "Bake renders into new textures for a target Mesh from multiple source Meshes");
-	REGISTER_MODELING_TOOL_COMMAND(BeginBakeMeshAttributeVertexTool, "BakeVtx", "Bake vertex colors for a target Mesh");
+	if ( bUseShortNames )
+	{
+		REGISTER_MODELING_TOOL_COMMAND(BeginBakeMeshAttributeMapsTool, "BakeTx", "Bake textures for a target Mesh");
+		REGISTER_MODELING_TOOL_COMMAND(BeginBakeMultiMeshAttributeMapsTool, "BakeAll", "Bake textures for a target Mesh from multiple source Meshes");
+		REGISTER_MODELING_TOOL_COMMAND(BeginBakeRenderCaptureTool, "BakeRC", "Bake renders into new textures for a target Mesh from multiple source Meshes");
+		REGISTER_MODELING_TOOL_COMMAND(BeginBakeMeshAttributeVertexTool, "BakeVtx", "Bake vertex colors for a target Mesh");
+	}
+	else
+	{
+		REGISTER_MODELING_TOOL_COMMAND(BeginBakeMeshAttributeMapsTool, "Bake Textures", "Bake textures for a target Mesh");
+		REGISTER_MODELING_TOOL_COMMAND(BeginBakeMultiMeshAttributeMapsTool, "Bake All", "Bake textures for a target Mesh from multiple source Meshes");
+		REGISTER_MODELING_TOOL_COMMAND(BeginBakeMeshAttributeVertexTool, "Bake Vertex Colors", "Bake vertex colors for a target Mesh");
+		REGISTER_MODELING_TOOL_COMMAND(BeginBakeRenderCaptureTool, "Bake RC", "Bake renders into new textures for a target Mesh from multiple source Meshes");
+	}
 
 	// Volumes
-	REGISTER_MODELING_TOOL_COMMAND(BeginVolumeToMeshTool, "Vol2Msh", "Convert a Volume to a new Mesh Object");
-	REGISTER_MODELING_TOOL_COMMAND(BeginMeshToVolumeTool, "Msh2Vol", "Convert a Mesh to a Volume");
-	if (!Settings->InRestrictiveMode())
+	if ( bUseShortNames )
 	{
-		REGISTER_MODELING_TOOL_COMMAND(BeginBspConversionTool, "BSPConv", "Convert BSP to a new Mesh Object");
+		REGISTER_MODELING_TOOL_COMMAND(BeginVolumeToMeshTool, "Vol2Msh", "Convert a Volume to a new Mesh Object");
+		REGISTER_MODELING_TOOL_COMMAND(BeginMeshToVolumeTool, "Msh2Vol", "Convert a Mesh to a Volume");
+		if (!Settings->InRestrictiveMode())
+		{
+			REGISTER_MODELING_TOOL_COMMAND(BeginBspConversionTool, "BSPConv", "Convert BSP to a new Mesh Object");
+		}
+		REGISTER_MODELING_TOOL_COMMAND(BeginPhysicsInspectorTool, "PInspct", "Inspect the physics/collision geometry for selected Meshes");
+		REGISTER_MODELING_TOOL_COMMAND(BeginSetCollisionGeometryTool, "Msh2Coll", "Convert selected Meshes to Simple Collision Geometry (for last selected)");
+		REGISTER_MODELING_TOOL_COMMAND(BeginExtractCollisionGeometryTool, "Coll2Msh", "Convert Simple Collision Geometry to a new Mesh Object");
+	}
+	else
+	{
+		REGISTER_MODELING_TOOL_COMMAND(BeginVolumeToMeshTool, "Volume To Mesh", "Convert a Volume to a new Mesh Object");
+		REGISTER_MODELING_TOOL_COMMAND(BeginMeshToVolumeTool, "Mesh To Volume", "Convert a Mesh to a Volume");
+		if (!Settings->InRestrictiveMode())
+		{
+			REGISTER_MODELING_TOOL_COMMAND(BeginBspConversionTool, "Convert BSPs", "Convert BSP to a new Mesh Object");
+		}
+		REGISTER_MODELING_TOOL_COMMAND(BeginPhysicsInspectorTool, "Inspect Collision", "Inspect the physics/collision geometry for selected Meshes");
+		REGISTER_MODELING_TOOL_COMMAND(BeginSetCollisionGeometryTool, "Mesh To Collision", "Convert selected Meshes to Simple Collision Geometry (for last selected)");
+		REGISTER_MODELING_TOOL_COMMAND(BeginExtractCollisionGeometryTool, "Collision To Mesh", "Convert Simple Collision Geometry to a new Mesh Object");
 	}
-	REGISTER_MODELING_TOOL_COMMAND(BeginPhysicsInspectorTool, "PInspct", "Inspect the physics/collision geometry for selected Meshes");
-	REGISTER_MODELING_TOOL_COMMAND(BeginSetCollisionGeometryTool, "Msh2Coll", "Convert selected Meshes to Simple Collision Geometry (for last selected)");
-	REGISTER_MODELING_TOOL_COMMAND(BeginExtractCollisionGeometryTool, "Coll2Msh", "Convert Simple Collision Geometry to a new Mesh Object");
 
 	// LODs
-	REGISTER_MODELING_TOOL_COMMAND(BeginLODManagerTool, "LODMgr", "Inspect the LODs of a Static Mesh Asset");
-	REGISTER_MODELING_TOOL_COMMAND(BeginGenerateStaticMeshLODAssetTool, "AutoLOD", "Automatically generate a simplified LOD with baked Textures/Materials for a Mesh");
-	REGISTER_MODELING_TOOL_COMMAND(BeginISMEditorTool, "ISMEd", "Edit the Instances of Instanced Static Mesh Components");
+	if ( bUseShortNames )
+	{
+		REGISTER_MODELING_TOOL_COMMAND(BeginLODManagerTool, "LODMgr", "Inspect the LODs of a Static Mesh Asset");
+		REGISTER_MODELING_TOOL_COMMAND(BeginGenerateStaticMeshLODAssetTool, "AutoLOD", "Automatically generate a simplified LOD with baked Textures/Materials for a Mesh");
+		REGISTER_MODELING_TOOL_COMMAND(BeginISMEditorTool, "ISMEd", "Edit the Instances of Instanced Static Mesh Components");
+	}
+	else
+	{
+		REGISTER_MODELING_TOOL_COMMAND(BeginLODManagerTool, "LOD Manager", "Inspect the LODs of a Static Mesh Asset");
+		REGISTER_MODELING_TOOL_COMMAND(BeginGenerateStaticMeshLODAssetTool, "AutoLOD", "Automatically generate a simplified LOD with baked Textures/Materials for a Mesh");
+		REGISTER_MODELING_TOOL_COMMAND(BeginISMEditorTool, "ISM Editor", "Edit the Instances of Instanced Static Mesh Components");
+	}
 
 	REGISTER_MODELING_TOOL_COMMAND(BeginAddPatchTool, "Patch", "Add Patch");
 	REGISTER_MODELING_TOOL_COMMAND(BeginShapeSprayTool, "Spray", "Shape Spray");
@@ -178,14 +326,30 @@ void FModelingToolsManagerCommands::RegisterCommands()
 	REGISTER_MODELING_TOOL_COMMAND(BeginPolyModelTool_Inset, "Inset", "Inset the current Selected Faces");
 	REGISTER_MODELING_TOOL_COMMAND(BeginPolyModelTool_Outset, "Outset", "Outset the current Selected Faces");
 	REGISTER_MODELING_TOOL_COMMAND(BeginPolyModelTool_CutFaces, "Cut", "Cut the current Selected Faces");
-	REGISTER_MODELING_TOOL_COMMAND(BeginPolyModelTool_InsertEdgeLoop, "ELoop", "Insert Edge Loops into the Selected Mesh");
 	REGISTER_MODELING_TOOL_COMMAND(BeginPolyModelTool_PushPull, "PushPull", "Push/Pull the current Selected Faces");
 	REGISTER_MODELING_TOOL_COMMAND(BeginPolyModelTool_Bevel, "Bevel", "Bevel the current Mesh Selection (Edges or Faces)");
+	if ( bUseShortNames )
+	{
+		REGISTER_MODELING_TOOL_COMMAND(BeginPolyModelTool_InsertEdgeLoop, "ELoop", "Insert Edge Loops into the Selected Mesh");
+	}
+	else
+	{
+		REGISTER_MODELING_TOOL_COMMAND(BeginPolyModelTool_InsertEdgeLoop, "Insert Loops", "Insert Edge Loops into the Selected Mesh");
+	}
+
+	if ( bUseShortNames )
+	{
+		REGISTER_MODELING_TOOL_COMMAND(BeginPolyModelTool_PolyEd, "PolyEd", "Select / Edit the current Mesh via PolyGroups");
+		REGISTER_MODELING_TOOL_COMMAND(BeginPolyModelTool_TriSel, "TriSel", "Select / Edit the current Mesh via Triangles");
+	}
+	else
+	{
+		REGISTER_MODELING_TOOL_COMMAND(BeginPolyModelTool_PolyEd, "PolyGroup Edit", "Select / Edit the current Mesh via PolyGroups");
+		REGISTER_MODELING_TOOL_COMMAND(BeginPolyModelTool_TriSel, "Triangle Edit", "Select / Edit the current Mesh via Triangles");
+	}
 
-	REGISTER_MODELING_TOOL_COMMAND(BeginPolyModelTool_PolyEd, "PolyEd", "Select / Edit the current Mesh via PolyGroups");
-	REGISTER_MODELING_TOOL_COMMAND(BeginPolyModelTool_TriSel, "TriSel", "Select / Edit the current Mesh via Triangles");
 
-	REGISTER_MODELING_TOOL_COMMAND(MeshSelectionModeAction_NoSelection, "None", "Disable Mesh Selection");
+	REGISTER_MODELING_TOOL_COMMAND(MeshSelectionModeAction_NoSelection, "None", "Disable Mesh Element Selection (Select Objects Only)");
 	REGISTER_MODELING_TOOL_COMMAND(MeshSelectionModeAction_MeshTriangles, "Tris", "Select Mesh Triangles");
 	REGISTER_MODELING_TOOL_COMMAND(MeshSelectionModeAction_MeshVertices, "Verts", "Select Mesh Vertices");
 	REGISTER_MODELING_TOOL_COMMAND(MeshSelectionModeAction_MeshEdges, "Edges", "Select Mesh Edges");
@@ -194,8 +358,15 @@ void FModelingToolsManagerCommands::RegisterCommands()
 	REGISTER_MODELING_TOOL_COMMAND(MeshSelectionModeAction_GroupEdges, "Borders", "Select Mesh PolyGroup Borders/Edges");
 
 	REGISTER_MODELING_ACTION_COMMAND(BeginSelectionAction_Delete, "Delete", "Delete the current Mesh Selection");
-	REGISTER_MODELING_ACTION_COMMAND(BeginSelectionAction_Disconnect, "Discon", "Disconnect the current Mesh Selection");
 	REGISTER_MODELING_ACTION_COMMAND(BeginSelectionAction_Retriangulate, "Clean", "Retriangulate the current Mesh or Mesh Selection");
+	if ( bUseShortNames )
+	{
+		REGISTER_MODELING_ACTION_COMMAND(BeginSelectionAction_Disconnect, "Discon", "Disconnect the current Mesh Selection");
+	}
+	else
+	{
+		REGISTER_MODELING_ACTION_COMMAND(BeginSelectionAction_Disconnect, "Disconnect", "Disconnect the current Mesh Selection");
+	}
 
 	REGISTER_MODELING_TOOL_COMMAND(BeginSelectionAction_Extrude, "Extrude", "Extrude the current Mesh Selection");
 	REGISTER_MODELING_TOOL_COMMAND(BeginSelectionAction_Offset, "Offset", "Offset the current Mesh Selection");
@@ -212,6 +383,7 @@ void FModelingToolsManagerCommands::RegisterCommands()
 	UI_COMMAND(AddToFavorites, "Add to Favorites", "Add to Favorites", EUserInterfaceActionType::Button, FInputChord());
 	UI_COMMAND(RemoveFromFavorites, "Remove from Favorites", "Remove from Favorites", EUserInterfaceActionType::Button, FInputChord());
 	UI_COMMAND(LoadFavoritesTools, "Favorites", "Favorites", EUserInterfaceActionType::RadioButton, FInputChord());
+	UI_COMMAND(LoadSelectionTools, "Selection", "Selection", EUserInterfaceActionType::RadioButton, FInputChord());
 	UI_COMMAND(LoadShapesTools, "Shapes", "Shapes", EUserInterfaceActionType::RadioButton, FInputChord());
 	UI_COMMAND(LoadCreateTools, "Create", "Create", EUserInterfaceActionType::RadioButton, FInputChord());
 	UI_COMMAND(LoadPolyTools, "PolyEdit", "PolyEdit", EUserInterfaceActionType::RadioButton, FInputChord());
diff --git a/Engine/Plugins/Editor/ModelingToolsEditorMode/Source/ModelingToolsEditorMode/Public/ModelingToolsEditorModeSettings.h b/Engine/Plugins/Editor/ModelingToolsEditorMode/Source/ModelingToolsEditorMode/Public/ModelingToolsEditorModeSettings.h
--- a/Engine/Plugins/Editor/ModelingToolsEditorMode/Source/ModelingToolsEditorMode/Public/ModelingToolsEditorModeSettings.h
+++ b/Engine/Plugins/Editor/ModelingToolsEditorMode/Source/ModelingToolsEditorMode/Public/ModelingToolsEditorModeSettings.h
@@ -296,6 +296,9 @@ public:
 
 public:
 
+	/** Toggle between the Legacy Modeling Mode Palette and the new UI (requires exiting and re-entering the Mode) */
+	UPROPERTY(config, EditAnywhere, Category = "Modeling Mode|UI Customization")
+	bool bUseLegacyModelingPalette = false;
 
 	/** Add the names of Modeling Mode Tool Palette Sections to have them appear at the top of the Tool Palette, in the order listed below. */
 	UPROPERTY(config, EditAnywhere, Category = "Modeling Mode|UI Customization")
diff --git a/Engine/Plugins/Editor/ModelingToolsEditorMode/Source/ModelingToolsEditorMode/Public/ModelingToolsManagerActions.h b/Engine/Plugins/Editor/ModelingToolsEditorMode/Source/ModelingToolsEditorMode/Public/ModelingToolsManagerActions.h
--- a/Engine/Plugins/Editor/ModelingToolsEditorMode/Source/ModelingToolsEditorMode/Public/ModelingToolsManagerActions.h
+++ b/Engine/Plugins/Editor/ModelingToolsEditorMode/Source/ModelingToolsEditorMode/Public/ModelingToolsManagerActions.h
@@ -169,6 +169,7 @@ public:
 	TSharedPtr<FUICommandInfo> AddToFavorites;
 	TSharedPtr<FUICommandInfo> RemoveFromFavorites;
 	TSharedPtr<FUICommandInfo> LoadFavoritesTools;
+	TSharedPtr<FUICommandInfo> LoadSelectionTools;
 	TSharedPtr<FUICommandInfo> LoadShapesTools;
 	TSharedPtr<FUICommandInfo> LoadCreateTools;
 	TSharedPtr<FUICommandInfo> LoadPolyTools;
