ModelingMode: fix up breakage around accessing UV Editor due to stream drift

diff --git a/Engine/Plugins/Editor/ModelingToolsEditorMode/Source/ModelingToolsEditorMode/Private/ModelingToolsEditorModeToolkit.cpp b/Engine/Plugins/Editor/ModelingToolsEditorMode/Source/ModelingToolsEditorMode/Private/ModelingToolsEditorModeToolkit.cpp
--- a/Engine/Plugins/Editor/ModelingToolsEditorMode/Source/ModelingToolsEditorMode/Private/ModelingToolsEditorModeToolkit.cpp
+++ b/Engine/Plugins/Editor/ModelingToolsEditorMode/Source/ModelingToolsEditorMode/Private/ModelingToolsEditorModeToolkit.cpp
@@ -320,10 +320,7 @@ void FModelingToolsEditorModeToolkit::RegisterPalettes()
 		Commands.BeginTransformUVIslandsTool,
 		Commands.BeginUVLayoutTool
 	});
-	IGeometryProcessingInterfacesModule& GeomProcInterfaces = FModuleManager::Get().LoadModuleChecked<
-		IGeometryProcessingInterfacesModule>("GeometryProcessingInterfaces");
-	IGeometryProcessing_UVEditorAssetEditor* UVEditorAPI = GeomProcInterfaces.GetUVEditorAssetEditorImplementation();
-	if (UVEditorAPI)
+	if (IModularFeatures::Get().IsModularFeatureAvailable(IUVEditorModularFeature::GetModularFeatureName()))
 	{
 		UVPaletteItems.Add(Commands.LaunchUVEditor);
 	}
