ModelingMode: fix localization text conflict

diff --git a/Engine/Plugins/Editor/ModelingToolsEditorMode/Source/ModelingToolsEditorMode/Private/ModelingToolsManagerActions.cpp b/Engine/Plugins/Editor/ModelingToolsEditorMode/Source/ModelingToolsEditorMode/Private/ModelingToolsManagerActions.cpp
--- a/Engine/Plugins/Editor/ModelingToolsEditorMode/Source/ModelingToolsEditorMode/Private/ModelingToolsManagerActions.cpp
+++ b/Engine/Plugins/Editor/ModelingToolsEditorMode/Source/ModelingToolsEditorMode/Private/ModelingToolsManagerActions.cpp
@@ -236,27 +236,24 @@ void FModelingToolsManagerCommands::RegisterCommands()
 	}
 
 	// UVs
+	REGISTER_MODELING_TOOL_COMMAND(BeginGlobalUVGenerateTool, "AutoUV", "Automatically unwrap and pack UVs");
+	// This is done directly, not with the REGISTER_ macro, since we don't want it added to the tool list or use a toggle button
+	UI_COMMAND(LaunchUVEditor, "UVEditor", "Launch UV asset editor", EUserInterfaceActionType::Button, FInputChord());
 	if ( bUseShortNames )
 	{
-		REGISTER_MODELING_TOOL_COMMAND(BeginGlobalUVGenerateTool, "AutoUV", "Automatically unwrap and pack UVs");
 		REGISTER_MODELING_TOOL_COMMAND(BeginGroupUVGenerateTool, "Unwrap", "Recompute UVs for existing UV islands or PolyGroups");
 		REGISTER_MODELING_TOOL_COMMAND(BeginUVProjectionTool, "Project", "Compute UVs via projecting to simple shapes");
 		REGISTER_MODELING_TOOL_COMMAND(BeginUVSeamEditTool, "SeamEd", "Add UV seams");
 		REGISTER_MODELING_TOOL_COMMAND(BeginTransformUVIslandsTool, "XFormUV", "Transform UV islands in UV space");
 		REGISTER_MODELING_TOOL_COMMAND(BeginUVLayoutTool, "Layout", "Transform and Repack existing UVs");
-		// This is done directly, not with the REGISTER_ macro, since we don't want it added to the tool list or use a toggle button
-		UI_COMMAND(LaunchUVEditor, "UVEditor", "Launch UV asset editor", EUserInterfaceActionType::Button, FInputChord());
 	}
 	else
 	{
-		REGISTER_MODELING_TOOL_COMMAND(BeginGlobalUVGenerateTool, "AutoUV", "Automatically unwrap and pack UVs");
 		REGISTER_MODELING_TOOL_COMMAND(BeginGroupUVGenerateTool, "UV Unwrap", "Recompute UVs for existing UV islands or PolyGroups");
 		REGISTER_MODELING_TOOL_COMMAND(BeginUVProjectionTool, "Project UVs", "Compute UVs via projecting to simple shapes");
 		REGISTER_MODELING_TOOL_COMMAND(BeginUVSeamEditTool, "Edit UV Seams", "Add UV seams");
 		REGISTER_MODELING_TOOL_COMMAND(BeginTransformUVIslandsTool, "Transform UVs", "Transform UV islands in UV space");
 		REGISTER_MODELING_TOOL_COMMAND(BeginUVLayoutTool, "Layout UVs", "Transform and Repack existing UVs");
-		// This is done directly, not with the REGISTER_ macro, since we don't want it added to the tool list or use a toggle button
-		UI_COMMAND(LaunchUVEditor, "UV Editor", "Launch UV asset editor", EUserInterfaceActionType::Button, FInputChord());
 	}
 
 	// Baking
