ModelingMode: disable mesh selection system by default.

diff --git a/Engine/Plugins/Editor/ModelingToolsEditorMode/Source/ModelingToolsEditorMode/Private/ModelingModeToolkit_Toolbars.cpp b/Engine/Plugins/Editor/ModelingToolsEditorMode/Source/ModelingToolsEditorMode/Private/ModelingModeToolkit_Toolbars.cpp
--- a/Engine/Plugins/Editor/ModelingToolsEditorMode/Source/ModelingToolsEditorMode/Private/ModelingModeToolkit_Toolbars.cpp
+++ b/Engine/Plugins/Editor/ModelingToolsEditorMode/Source/ModelingToolsEditorMode/Private/ModelingModeToolkit_Toolbars.cpp
@@ -411,7 +411,7 @@ void FModelingToolsEditorModeToolkit::ExtendSecondaryModeToolbar(UToolMenu *InMo
 	}
 
 	UModelingToolsEditorModeSettings* ModelingModeSettings = GetMutableDefault<UModelingToolsEditorModeSettings>();
-	const bool bEnableSelectionUI = ModelingModeSettings && ModelingModeSettings->bEnableMeshSelections;
+	const bool bEnableSelectionUI = ModelingModeSettings && ModelingModeSettings->GetMeshSelectionsEnabled();
 	if ( !bEnableSelectionUI )
 	{
 		return;
diff --git a/Engine/Plugins/Editor/ModelingToolsEditorMode/Source/ModelingToolsEditorMode/Private/ModelingToolsEditorMode.cpp b/Engine/Plugins/Editor/ModelingToolsEditorMode/Source/ModelingToolsEditorMode/Private/ModelingToolsEditorMode.cpp
--- a/Engine/Plugins/Editor/ModelingToolsEditorMode/Source/ModelingToolsEditorMode/Private/ModelingToolsEditorMode.cpp
+++ b/Engine/Plugins/Editor/ModelingToolsEditorMode/Source/ModelingToolsEditorMode/Private/ModelingToolsEditorMode.cpp
@@ -401,7 +401,7 @@ void UModelingToolsEditorMode::Enter()
 	SceneSnappingManager = UE::Geometry::FindModelingSceneSnappingManager(GetToolManager());
 
 	// register selection manager, if this feature is enabled in the mode settings
-	if (ModelingModeSettings && ModelingModeSettings->bEnableMeshSelections)
+	if (ModelingModeSettings && ModelingModeSettings->GetMeshSelectionsEnabled())
 	{
 		// set up SelectionManager and register known factory types
 		SelectionManager = NewObject<UGeometrySelectionManager>(GetToolManager());
diff --git a/Engine/Plugins/Editor/ModelingToolsEditorMode/Source/ModelingToolsEditorMode/Private/ModelingToolsEditorModeToolkit.cpp b/Engine/Plugins/Editor/ModelingToolsEditorMode/Source/ModelingToolsEditorMode/Private/ModelingToolsEditorModeToolkit.cpp
--- a/Engine/Plugins/Editor/ModelingToolsEditorMode/Source/ModelingToolsEditorMode/Private/ModelingToolsEditorModeToolkit.cpp
+++ b/Engine/Plugins/Editor/ModelingToolsEditorMode/Source/ModelingToolsEditorMode/Private/ModelingToolsEditorModeToolkit.cpp
@@ -153,6 +153,7 @@ void FModelingToolsEditorModeToolkit::RegisterPalettes()
 
 	UEdMode* ScriptableMode = GetScriptableEditorMode().Get();
 	UModelingToolsModeCustomizationSettings* UISettings = GetMutableDefault<UModelingToolsModeCustomizationSettings>();
+	UModelingToolsEditorModeSettings* ModelingModeSettings = GetMutableDefault<UModelingToolsEditorModeSettings>();
 	
 	ToolkitBuilder = MakeShared<FToolkitBuilder>(
 				ScriptableMode->GetModeInfo().ToolbarCustomizationName, 
@@ -196,26 +197,28 @@ void FModelingToolsEditorModeToolkit::RegisterPalettes()
 	ToolkitBuilder->AddPalette(
 		MakeShareable( new FToolPalette( Commands.LoadCreateTools.ToSharedRef(), CreatePaletteItems ) ) );
 
+	if ( ModelingModeSettings && ModelingModeSettings->GetMeshSelectionsEnabled() )
+	{
+		TArray<TSharedPtr<FUICommandInfo>> SelectionPaletteItems({
+			Commands.BeginSelectionAction_Delete,
+			Commands.BeginSelectionAction_Extrude,
+			Commands.BeginSelectionAction_Offset,
 
-	TArray<TSharedPtr<FUICommandInfo>> SelectionPaletteItems({
-		Commands.BeginSelectionAction_Delete,
-		Commands.BeginSelectionAction_Extrude,
-		Commands.BeginSelectionAction_Offset,
-
-		Commands.BeginPolyModelTool_PushPull,
-		Commands.BeginPolyModelTool_Inset,
-		Commands.BeginPolyModelTool_Outset,
-		Commands.BeginPolyModelTool_CutFaces,
-		Commands.BeginPolyModelTool_Bevel,
+			Commands.BeginPolyModelTool_PushPull,
+			Commands.BeginPolyModelTool_Inset,
+			Commands.BeginPolyModelTool_Outset,
+			Commands.BeginPolyModelTool_CutFaces,
+			Commands.BeginPolyModelTool_Bevel,
 
-		Commands.BeginPolyModelTool_InsertEdgeLoop,
-		Commands.BeginSelectionAction_Retriangulate,
+			Commands.BeginPolyModelTool_InsertEdgeLoop,
+			Commands.BeginSelectionAction_Retriangulate,
 
-		Commands.BeginPolyModelTool_PolyEd,
-		Commands.BeginPolyModelTool_TriSel
-	});
-	ToolkitBuilder->AddPalette( 
-		MakeShareable( new FToolPalette( Commands.LoadSelectionTools.ToSharedRef(), SelectionPaletteItems ) ) );
+			Commands.BeginPolyModelTool_PolyEd,
+			Commands.BeginPolyModelTool_TriSel
+		});
+		ToolkitBuilder->AddPalette( 
+			MakeShareable( new FToolPalette( Commands.LoadSelectionTools.ToSharedRef(), SelectionPaletteItems ) ) );
+	}
 
 
 	TArray<TSharedPtr<FUICommandInfo>> TransformPaletteItems({
@@ -1284,7 +1287,7 @@ void FModelingToolsEditorModeToolkit::GetToolPaletteNames(TArray<FName>& Palette
 	}
 
 	const UModelingToolsEditorModeSettings* ModelingModeSettings = GetDefault<UModelingToolsEditorModeSettings>();
-	bool bEnableSelectionUI = ModelingModeSettings && ModelingModeSettings->bEnableMeshSelections;
+	bool bEnableSelectionUI = ModelingModeSettings && ModelingModeSettings->GetMeshSelectionsEnabled();
 	if (bEnableSelectionUI)
 	{
 		if (bShowActiveSelectionActions)
diff --git a/Engine/Plugins/Editor/ModelingToolsEditorMode/Source/ModelingToolsEditorMode/Public/ModelingToolsEditorModeSettings.h b/Engine/Plugins/Editor/ModelingToolsEditorMode/Source/ModelingToolsEditorMode/Public/ModelingToolsEditorModeSettings.h
--- a/Engine/Plugins/Editor/ModelingToolsEditorMode/Source/ModelingToolsEditorMode/Public/ModelingToolsEditorModeSettings.h
+++ b/Engine/Plugins/Editor/ModelingToolsEditorMode/Source/ModelingToolsEditorMode/Public/ModelingToolsEditorModeSettings.h
@@ -146,14 +146,22 @@ public:
 	// Selection
 	//
 
-	// old preference that should be ignored now
-	UPROPERTY()
+	// old preference for mesh selection system that will be disabled in 5.3
+	UPROPERTY(config, EditAnywhere, Category = "Modeling Mode|Selection", meta = (DisplayName="Enable Mesh Selection UI"))
 	bool bEnablePersistentSelections = false;
 
+	// new default-enabled preference for the mesh selection system that we will switch to in 5.3
 	/** Enable/Disable new Mesh Selection System (Experimental). */
-	UPROPERTY(config, EditAnywhere, Category = "Modeling Mode|Selection", meta = (DisplayName="Enable Mesh Selection UI"))
+	//UPROPERTY(config, EditAnywhere, Category = "Modeling Mode|Selection", meta = (DisplayName="Enable Mesh Selection UI"))
+	UPROPERTY()
 	bool bEnableMeshSelections = true;
 
+	virtual bool GetMeshSelectionsEnabled() const
+	{
+		return bEnablePersistentSelections;
+		//return bEnableMeshSelections;
+	}
+
 
 	DECLARE_MULTICAST_DELEGATE_TwoParams(UModelingToolsEditorModeSettingsModified, UObject*, FProperty*);
 	UModelingToolsEditorModeSettingsModified OnModified;
