MeshModelingTools: Move various settings from the secondary toolbar to the cog icon in quick settings, since the secondary toolbar is not enabled by default in restrictive mode.

diff --git a/Engine/Plugins/Editor/ModelingToolsEditorMode/Source/ModelingToolsEditorMode/Private/ModelingModeToolkit_Toolbars.cpp b/Engine/Plugins/Editor/ModelingToolsEditorMode/Source/ModelingToolsEditorMode/Private/ModelingModeToolkit_Toolbars.cpp
--- a/Engine/Plugins/Editor/ModelingToolsEditorMode/Source/ModelingToolsEditorMode/Private/ModelingModeToolkit_Toolbars.cpp
+++ b/Engine/Plugins/Editor/ModelingToolsEditorMode/Source/ModelingToolsEditorMode/Private/ModelingModeToolkit_Toolbars.cpp
@@ -351,7 +351,7 @@ void MakeSubMenu_Selection_MeshType(FModelingToolsEditorModeToolkit* Toolkit, FM
 
 }
 
-TSharedRef<SWidget> MakeMenu_ModelingModeConfigSettings(FModelingToolsEditorModeToolkit* Toolkit)
+TSharedRef<SWidget> MakeMenu_SelectionConfigSettings(FModelingToolsEditorModeToolkit* Toolkit)
 {
 	FMenuBuilder MenuBuilder(true, TSharedPtr<FUICommandList>());
 
@@ -363,12 +363,26 @@ TSharedRef<SWidget> MakeMenu_ModelingModeConfigSettings(FModelingToolsEditorMode
 	MakeSubMenu_Selection_MeshType(Toolkit, MenuBuilder);
 	MenuBuilder.EndSection();
 
-	MenuBuilder.BeginSection("Section_Selection", LOCTEXT("Section_Selection", "Selection"));
+	TSharedRef<SWidget> MenuWidget = MenuBuilder.MakeWidget();
+	return MenuWidget;
+}
+
+
+} // end namespace UELocal
+
+
+TSharedRef<SWidget> FModelingToolsEditorModeToolkit::MakeMenu_ModelingModeConfigSettings()
+{
+	using namespace UELocal;
+
+	FMenuBuilder MenuBuilder(true, TSharedPtr<FUICommandList>());
+
+	MenuBuilder.BeginSection("Section_Selection", LOCTEXT("Section_Selection", "Instances"));
 	MakeSubMenu_ModeToggles(MenuBuilder);
 	MenuBuilder.EndSection();
 
 	MenuBuilder.BeginSection("Section_Gizmo", LOCTEXT("Section_Gizmo", "Gizmo"));
-	MakeSubMenu_GizmoVisibilityMode(Toolkit, MenuBuilder);
+	MakeSubMenu_GizmoVisibilityMode(this, MenuBuilder);
 	MenuBuilder.EndSection();
 
 	MenuBuilder.BeginSection("Section_MeshObjects", LOCTEXT("Section_MeshObjects", "New Mesh Objects"));
@@ -396,9 +410,6 @@ TSharedRef<SWidget> MakeMenu_ModelingModeConfigSettings(FModelingToolsEditorMode
 	return MenuWidget;
 }
 
-} // end namespace UELocal
-
-
 void FModelingToolsEditorModeToolkit::ExtendSecondaryModeToolbar(UToolMenu *InModeToolbarMenu)
 {
 	/*
@@ -492,7 +503,7 @@ void FModelingToolsEditorModeToolkit::ExtendSecondaryModeToolbar(UToolMenu *InMo
 		.ComboButtonStyle(FAppStyle::Get(), "SimpleComboButton")
 		.OnGetMenuContent(FOnGetContent::CreateLambda([this]()
 		{
-			return UELocal::MakeMenu_ModelingModeConfigSettings(this);
+			return UELocal::MakeMenu_SelectionConfigSettings(this);
 		}))
 		.ContentPadding(FMargin(3.0f, 1.0f))
 		.ButtonContent()
diff --git a/Engine/Plugins/Editor/ModelingToolsEditorMode/Source/ModelingToolsEditorMode/Private/ModelingToolsEditorModeToolkit.cpp b/Engine/Plugins/Editor/ModelingToolsEditorMode/Source/ModelingToolsEditorMode/Private/ModelingToolsEditorModeToolkit.cpp
--- a/Engine/Plugins/Editor/ModelingToolsEditorMode/Source/ModelingToolsEditorMode/Private/ModelingToolsEditorModeToolkit.cpp
+++ b/Engine/Plugins/Editor/ModelingToolsEditorMode/Source/ModelingToolsEditorMode/Private/ModelingToolsEditorModeToolkit.cpp
@@ -707,53 +707,70 @@ TSharedPtr<SWidget> FModelingToolsEditorModeToolkit::MakeAssetConfigPanel()
 				[ AssetLODMode->AsShared() ]
 		];
 
+	SHorizontalBox::FArguments AssetOptionsArgs;
 	if (Settings->InRestrictiveMode())
 	{
 		NewAssetPath = SNew(SEditableTextBox)
 		.Text(this, &FModelingToolsEditorModeToolkit::GetRestrictiveModeAutoGeneratedAssetPathText)
 		.OnTextCommitted(this, &FModelingToolsEditorModeToolkit::OnRestrictiveModeAutoGeneratedAssetPathTextCommitted);
 
-		Content->AddSlot()[
-			SNew(SHorizontalBox)
-			+ SHorizontalBox::Slot().Padding(0, 2, 2, 2).HAlign(HAlign_Right).VAlign(VAlign_Center).AutoWidth()
-			[
-				SNew(STextBlock)
-				.Text(LOCTEXT("NewAssetPathLabel", "New Asset Path"))
-				.ToolTipText(LOCTEXT("NewAssetPathToolTip", "Path used for storing newly generated assets within the project folder."))
-			]
-			+ SHorizontalBox::Slot().HAlign(HAlign_Fill).Padding(0).FillWidth(1.0f)
-			[
-				NewAssetPath->AsShared()
-			]
-			+ SHorizontalBox::Slot().HAlign(HAlign_Right).Padding(0, 0, 0, 0).AutoWidth()
-			[
-				SNew(SSimpleButton)
-				.OnClicked_Lambda( [this]() { SelectNewAssetPath(); return FReply::Handled(); } )
-				.Icon(FAppStyle::Get().GetBrush("Icons.FolderOpen"))
-			]
+		AssetOptionsArgs
+		+ SHorizontalBox::Slot().Padding(0, 2, 2, 2).HAlign(HAlign_Right).VAlign(VAlign_Center).AutoWidth()
+		[
+			SNew(STextBlock)
+			.Text(LOCTEXT("NewAssetPathLabel", "New Asset Path"))
+			.ToolTipText(LOCTEXT("NewAssetPathToolTip", "Path used for storing newly generated assets within the project folder."))
+		]
+		+ SHorizontalBox::Slot().HAlign(HAlign_Fill).Padding(0).FillWidth(1.0f)
+		[
+			NewAssetPath->AsShared()
+		]
+		+ SHorizontalBox::Slot().HAlign(HAlign_Right).Padding(0, 0, 0, 0).AutoWidth()
+		[
+			SNew(SSimpleButton)
+			.OnClicked_Lambda( [this]() { SelectNewAssetPath(); return FReply::Handled(); } )
+			.Icon(FAppStyle::Get().GetBrush("Icons.FolderOpen"))
 		];
 	}
 	else
 	{
-		Content->AddSlot()[
-			SNew(SHorizontalBox)
-			+ SHorizontalBox::Slot().Padding(0, 2, 2, 2).HAlign(HAlign_Right).VAlign(VAlign_Center).AutoWidth()
-			[
-				SNew(STextBlock)
-				.Text(LOCTEXT("AssetLocationLabel", "New Asset Location"))
-			]
-			+ SHorizontalBox::Slot().HAlign(HAlign_Fill).Padding(0).FillWidth(1.0f)
-			[
-				AssetLocationMode->AsShared()
-			]
-			+ SHorizontalBox::Slot().HAlign(HAlign_Right).Padding(0, 0, 0, 0).AutoWidth()
-			[
-				SNew(SSimpleButton)
-				.OnClicked_Lambda( [this]() { OnShowAssetSettings(); return FReply::Handled(); } )
-				.Icon(FAppStyle::Get().GetBrush("Icons.Settings"))
-			]
+		AssetOptionsArgs
+		+ SHorizontalBox::Slot().Padding(0, 2, 2, 2).HAlign(HAlign_Right).VAlign(VAlign_Center).AutoWidth()
+		[
+			SNew(STextBlock)
+			.Text(LOCTEXT("AssetLocationLabel", "New Asset Location"))
+		]
+		+ SHorizontalBox::Slot().HAlign(HAlign_Fill).Padding(0).FillWidth(1.0f)
+		[
+			AssetLocationMode->AsShared()
 		];
 	}
+	
+	// Regardless of restrictive mode or not, we add a cog icon to the same row with some common settings.
+	AssetOptionsArgs
+	+ SHorizontalBox::Slot().HAlign(HAlign_Right).Padding(0, 0, 0, 0).AutoWidth()
+	[
+		SNew(SComboButton)
+		.HasDownArrow(false)
+		.MenuPlacement(EMenuPlacement::MenuPlacement_MenuRight)
+		.ComboButtonStyle(FAppStyle::Get(), "SimpleComboButton")
+		.OnGetMenuContent(FOnGetContent::CreateLambda([this]()
+		{
+			return MakeMenu_ModelingModeConfigSettings();
+		}))
+		.ContentPadding(FMargin(3.0f, 1.0f))
+		.ButtonContent()
+		[
+			SNew(SImage)
+			.Image(FModelingToolsEditorModeStyle::Get()->GetBrush("ModelingMode.DefaultSettings"))
+			.ColorAndOpacity(FSlateColor::UseForeground())
+		]
+	];
+
+	Content->AddSlot()[
+		SArgumentNew(AssetOptionsArgs, SHorizontalBox)
+	];
+
 
 	TSharedPtr<SExpandableArea> AssetConfigPanel = SNew(SExpandableArea)
 		.HeaderPadding(FMargin(2.0f))
diff --git a/Engine/Plugins/Editor/ModelingToolsEditorMode/Source/ModelingToolsEditorMode/Public/ModelingToolsEditorModeToolkit.h b/Engine/Plugins/Editor/ModelingToolsEditorMode/Source/ModelingToolsEditorMode/Public/ModelingToolsEditorModeToolkit.h
--- a/Engine/Plugins/Editor/ModelingToolsEditorMode/Source/ModelingToolsEditorMode/Public/ModelingToolsEditorModeToolkit.h
+++ b/Engine/Plugins/Editor/ModelingToolsEditorMode/Source/ModelingToolsEditorMode/Public/ModelingToolsEditorModeToolkit.h
@@ -128,6 +128,8 @@ private:
 	void MakeGizmoNumericalUIOverlayWidget();
 	TSharedPtr<STransformGizmoNumericalUIOverlay> GizmoNumericalUIOverlayWidget;
 
+	TSharedRef<SWidget> MakeMenu_ModelingModeConfigSettings();
+
 	TSharedPtr<STextBlock> ModeWarningArea;
 	TSharedPtr<STextBlock> ModeHeaderArea;
 	TSharedPtr<STextBlock> ToolWarningArea;
