MeshModelingToolset: Disable BakeRC tool in restrictive mode

diff --git a/Engine/Plugins/Editor/ModelingToolsEditorMode/Source/ModelingToolsEditorMode/Private/ModelingToolsEditorMode.cpp b/Engine/Plugins/Editor/ModelingToolsEditorMode/Source/ModelingToolsEditorMode/Private/ModelingToolsEditorMode.cpp
--- a/Engine/Plugins/Editor/ModelingToolsEditorMode/Source/ModelingToolsEditorMode/Private/ModelingToolsEditorMode.cpp
+++ b/Engine/Plugins/Editor/ModelingToolsEditorMode/Source/ModelingToolsEditorMode/Private/ModelingToolsEditorMode.cpp
@@ -714,8 +714,11 @@ void UModelingToolsEditorMode::Enter()
 	auto BakeMultiMeshAttributeMapsToolBuilder = NewObject<UBakeMultiMeshAttributeMapsToolBuilder>();
 	RegisterTool(ToolManagerCommands.BeginBakeMultiMeshAttributeMapsTool, TEXT("BeginBakeMultiMeshAttributeMapsTool"), BakeMultiMeshAttributeMapsToolBuilder);
 
-	auto BakeRenderCaptureToolBuilder = NewObject<UBakeRenderCaptureToolBuilder>();
-	RegisterTool(ToolManagerCommands.BeginBakeRenderCaptureTool, TEXT("BeginBakeRenderCaptureTool"), BakeRenderCaptureToolBuilder);
+	// BakeRC is disabled in Restrictive Mode.
+	if (ToolManagerCommands.BeginBakeRenderCaptureTool)
+	{
+		RegisterTool(ToolManagerCommands.BeginBakeRenderCaptureTool, TEXT("BeginBakeRenderCaptureTool"), NewObject<UBakeRenderCaptureToolBuilder>());
+	}
 
 	auto BakeMeshAttributeVertexToolBuilder = NewObject<UBakeMeshAttributeVertexToolBuilder>();
 	RegisterTool(ToolManagerCommands.BeginBakeMeshAttributeVertexTool, TEXT("BeginBakeMeshAttributeVertexTool"), BakeMeshAttributeVertexToolBuilder);
diff --git a/Engine/Plugins/Editor/ModelingToolsEditorMode/Source/ModelingToolsEditorMode/Private/ModelingToolsEditorModeToolkit.cpp b/Engine/Plugins/Editor/ModelingToolsEditorMode/Source/ModelingToolsEditorMode/Private/ModelingToolsEditorModeToolkit.cpp
--- a/Engine/Plugins/Editor/ModelingToolsEditorMode/Source/ModelingToolsEditorMode/Private/ModelingToolsEditorModeToolkit.cpp
+++ b/Engine/Plugins/Editor/ModelingToolsEditorMode/Source/ModelingToolsEditorMode/Private/ModelingToolsEditorModeToolkit.cpp
@@ -307,9 +307,12 @@ void FModelingToolsEditorModeToolkit::RegisterPalettes()
 	TArray<TSharedPtr<FUICommandInfo>> BakingPaletteItems({
 		Commands.BeginBakeMeshAttributeMapsTool,
 		Commands.BeginBakeMultiMeshAttributeMapsTool,
-		Commands.BeginBakeMeshAttributeVertexTool,
-		Commands.BeginBakeRenderCaptureTool
+		Commands.BeginBakeMeshAttributeVertexTool
 	});
+	if (Commands.BeginBakeRenderCaptureTool)
+	{
+		BakingPaletteItems.Add(Commands.BeginBakeRenderCaptureTool);
+	}
 	ToolkitBuilder->AddPalette(
 		MakeShareable( new FToolPalette( Commands.LoadBakingTools.ToSharedRef(), BakingPaletteItems ) ) );
 
@@ -1537,7 +1540,10 @@ void FModelingToolsEditorModeToolkit::BuildToolPalette(FName PaletteIndex, class
 		ToolbarBuilder.AddToolBarButton(Commands.BeginBakeMeshAttributeMapsTool);
 		ToolbarBuilder.AddToolBarButton(Commands.BeginBakeMultiMeshAttributeMapsTool);
 		ToolbarBuilder.AddToolBarButton(Commands.BeginBakeMeshAttributeVertexTool);
-		ToolbarBuilder.AddToolBarButton(Commands.BeginBakeRenderCaptureTool);
+		if (Commands.BeginBakeRenderCaptureTool)
+		{
+			ToolbarBuilder.AddToolBarButton(Commands.BeginBakeRenderCaptureTool);
+		}
 	}
 	else if (PaletteIndex == UVTabName)
 	{
diff --git a/Engine/Plugins/Editor/ModelingToolsEditorMode/Source/ModelingToolsEditorMode/Private/ModelingToolsManagerActions.cpp b/Engine/Plugins/Editor/ModelingToolsEditorMode/Source/ModelingToolsEditorMode/Private/ModelingToolsManagerActions.cpp
--- a/Engine/Plugins/Editor/ModelingToolsEditorMode/Source/ModelingToolsEditorMode/Private/ModelingToolsManagerActions.cpp
+++ b/Engine/Plugins/Editor/ModelingToolsEditorMode/Source/ModelingToolsEditorMode/Private/ModelingToolsManagerActions.cpp
@@ -261,15 +261,21 @@ void FModelingToolsManagerCommands::RegisterCommands()
 	{
 		REGISTER_MODELING_TOOL_COMMAND(BeginBakeMeshAttributeMapsTool, "BakeTx", "Bake textures for a target Mesh");
 		REGISTER_MODELING_TOOL_COMMAND(BeginBakeMultiMeshAttributeMapsTool, "BakeAll", "Bake textures for a target Mesh from multiple source Meshes");
-		REGISTER_MODELING_TOOL_COMMAND(BeginBakeRenderCaptureTool, "BakeRC", "Bake renders into new textures for a target Mesh from multiple source Meshes");
 		REGISTER_MODELING_TOOL_COMMAND(BeginBakeMeshAttributeVertexTool, "BakeVtx", "Bake vertex colors for a target Mesh");
+		if (!Settings->InRestrictiveMode())
+		{
+			REGISTER_MODELING_TOOL_COMMAND(BeginBakeRenderCaptureTool, "BakeRC", "Bake renders into new textures for a target Mesh from multiple source Meshes");
+		}
 	}
 	else
 	{
 		REGISTER_MODELING_TOOL_COMMAND(BeginBakeMeshAttributeMapsTool, "Bake Textures", "Bake textures for a target Mesh");
 		REGISTER_MODELING_TOOL_COMMAND(BeginBakeMultiMeshAttributeMapsTool, "Bake All", "Bake textures for a target Mesh from multiple source Meshes");
 		REGISTER_MODELING_TOOL_COMMAND(BeginBakeMeshAttributeVertexTool, "Bake Vertex Colors", "Bake vertex colors for a target Mesh");
-		REGISTER_MODELING_TOOL_COMMAND(BeginBakeRenderCaptureTool, "Bake RC", "Bake renders into new textures for a target Mesh from multiple source Meshes");
+		if (!Settings->InRestrictiveMode())
+		{
+			REGISTER_MODELING_TOOL_COMMAND(BeginBakeRenderCaptureTool, "Bake RC", "Bake renders into new textures for a target Mesh from multiple source Meshes");
+		}
 	}
 
 	// Volumes
