Toolkit and toolbar styling

diff --git a/Engine/Plugins/Editor/ModelingToolsEditorMode/Source/ModelingToolsEditorMode/Private/ModelingToolsEditorModeToolkit.cpp b/Engine/Plugins/Editor/ModelingToolsEditorMode/Source/ModelingToolsEditorMode/Private/ModelingToolsEditorModeToolkit.cpp
--- a/Engine/Plugins/Editor/ModelingToolsEditorMode/Source/ModelingToolsEditorMode/Private/ModelingToolsEditorModeToolkit.cpp
+++ b/Engine/Plugins/Editor/ModelingToolsEditorMode/Source/ModelingToolsEditorMode/Private/ModelingToolsEditorModeToolkit.cpp
@@ -154,10 +154,12 @@ void FModelingToolsEditorModeToolkit::RegisterPalettes()
 	UEdMode* ScriptableMode = GetScriptableEditorMode().Get();
 	UModelingToolsModeCustomizationSettings* UISettings = GetMutableDefault<UModelingToolsModeCustomizationSettings>();
 	UModelingToolsEditorModeSettings* ModelingModeSettings = GetMutableDefault<UModelingToolsEditorModeSettings>();
-	
+
+	ToolkitSections = MakeShared<FToolkitSections>();
 	ToolkitBuilder = MakeShared<FToolkitBuilder>(
 				ScriptableMode->GetModeInfo().ToolbarCustomizationName, 
-				GetToolkitCommands());
+				GetToolkitCommands(),
+				ToolkitSections);
 
 	
 	const TSharedPtr<FEditablePalette> FavoritesPalette =
@@ -375,12 +377,7 @@ void FModelingToolsEditorModeToolkit::RegisterPalettes()
 	ToolkitBuilder->AddPalette(
 		MakeShareable( new FToolPalette( Commands.LoadLodsTools.ToSharedRef(), LODPaletteItems)) );
 
-
-	ToolkitBuilder->OnToolEnded.BindLambda([this]()
-	{
-		GetScriptableEditorMode()->GetToolManager(EToolsContextScope::EdMode)->DeactivateTool(EToolSide::Left, EToolShutdownType::Cancel);
-	});
-
+	ToolkitBuilder->SetActivePaletteOnLoad(Commands.LoadShapesTools.Get());
 	ToolkitBuilder->UpdateWidget();
 }
 
@@ -394,20 +391,9 @@ void FModelingToolsEditorModeToolkit::Init(const TSharedPtr<IToolkitHost>& InitT
 	// ModeDetailsView created, that we need to add to the ToolkitWidget. So, we will create the Widget
 	// here but only add the rows to it after we call Init()
 
-	TSharedPtr<SHorizontalBox> ToolkitWidgetHBox;
-	TSharedPtr<SVerticalBox> ToolkitWidgetVBox = SNew(SVerticalBox);
+	const TSharedPtr<SVerticalBox> ToolkitWidgetVBox = SNew(SVerticalBox);
 	
-	if ( bUsesToolkitBuilder )
-	{
-		ToolkitWidgetHBox = SNew(SHorizontalBox);
-		SAssignNew(ToolkitWidget, SBorder)
-			.HAlign(HAlign_Fill)
-			.Padding(4)
-			[
-				ToolkitWidgetHBox->AsShared()
-			];		
-	}
-	else
+	if ( !bUsesToolkitBuilder )
 	{
 		SAssignNew(ToolkitWidget, SBorder)
 			.HAlign(HAlign_Fill)
@@ -454,54 +440,44 @@ void FModelingToolsEditorModeToolkit::Init(const TSharedPtr<IToolkitHost>& InitT
 	if (bUsesToolkitBuilder)
 	{
 		RegisterPalettes();
-
-		ToolkitWidgetVBox->AddSlot().AutoHeight().HAlign(HAlign_Fill).Padding(5)
-			[
-				ToolkitBuilder->GetToolPaletteWidget()
-			];
 	}
 	else
 	{
 		ToolkitWidgetVBox->AddSlot().AutoHeight().HAlign(HAlign_Fill).Padding(5)
 			[
 				ModeHeaderArea->AsShared()
-			];		
-	}
-		
-	ToolkitWidgetVBox->AddSlot().AutoHeight().HAlign(HAlign_Fill).Padding(5)
-	[
-		ToolWarningArea->AsShared()
-	];	
-	ToolkitWidgetVBox->AddSlot().HAlign(HAlign_Fill).FillHeight(1.f)
-	[
-		ModeDetailsView->AsShared()
-	];
-	ToolkitWidgetVBox->AddSlot().AutoHeight().HAlign(HAlign_Fill).VAlign(VAlign_Bottom).Padding(5)
+			];
+		ToolkitWidgetVBox->AddSlot().AutoHeight().HAlign(HAlign_Fill).Padding(5)
+		[
+			ToolWarningArea->AsShared()
+		];	
+		ToolkitWidgetVBox->AddSlot().HAlign(HAlign_Fill).FillHeight(1.f)
 		[
-			MakeAssetConfigPanel()->AsShared()
+			ModeDetailsView->AsShared()
 		];
+		ToolkitWidgetVBox->AddSlot().AutoHeight().HAlign(HAlign_Fill).VAlign(VAlign_Bottom).Padding(5)
+		[
+		MakeAssetConfigPanel()->AsShared()];
+	}
+		
 
 	ClearNotification();
 	ClearWarning();
 
-	if (bUsesToolkitBuilder)
+	if (HasToolkitBuilder())
 	{
-		ToolkitWidgetHBox->AddSlot()
-			.FillWidth(0.5f)
-			.MaxWidth(48)
-			.VAlign(VAlign_Fill)
-			.Padding(0)
-			[
-			ToolkitBuilder->CreateToolbarWidget()
-			];
-	
-		ToolkitWidgetHBox->AddSlot()
-			.FillWidth(4.5f)
-			.VAlign(VAlign_Fill)
-			.Padding(5)
-			[
-				ToolkitWidgetVBox->AsShared()
-			];
+		ToolkitSections->ModeWarningArea = ModeWarningArea;
+		ToolkitSections->DetailsView = ModeDetailsView;
+		ToolkitSections->ToolWarningArea = ToolWarningArea;
+		ToolkitSections->Footer = MakeAssetConfigPanel();
+
+		SAssignNew(ToolkitWidget, SBorder)
+		.HAlign(HAlign_Fill)
+		.Padding(0)
+		.BorderImage(FAppStyle::GetBrush("ToolPanel.GroupBorder"))
+		[
+			ToolkitBuilder->GenerateWidget()->AsShared()
+		];	
 	}
 
 	ActiveToolName = FText::GetEmpty();
@@ -702,11 +678,12 @@ TSharedPtr<SWidget> FModelingToolsEditorModeToolkit::MakeAssetConfigPanel()
 
 	const TSharedPtr<SVerticalBox> Content = SNew(SVerticalBox)
 	+ SVerticalBox::Slot().HAlign(HAlign_Fill)
+	.Padding(0, 8, 0, 4)
 		[
 			SNew(SHorizontalBox)
-			+ SHorizontalBox::Slot().Padding(0, 2, 2, 2).HAlign(HAlign_Right).VAlign(VAlign_Center).AutoWidth()
+			+ SHorizontalBox::Slot().Padding(0).HAlign(HAlign_Left).VAlign(VAlign_Center).FillWidth(2.f)
 				[ AssetLODModeLabel->AsShared() ]
-			+ SHorizontalBox::Slot().Padding(0).FillWidth(4.0f)
+			+ SHorizontalBox::Slot().Padding(0).HAlign(HAlign_Fill).Padding(0).FillWidth(4.f)
 				[ AssetLODMode->AsShared() ]
 		];
 
@@ -718,69 +695,64 @@ TSharedPtr<SWidget> FModelingToolsEditorModeToolkit::MakeAssetConfigPanel()
 		.OnTextCommitted(this, &FModelingToolsEditorModeToolkit::OnRestrictiveModeAutoGeneratedAssetPathTextCommitted);
 
 		AssetOptionsArgs
-		+ SHorizontalBox::Slot().Padding(0, 2, 2, 2).HAlign(HAlign_Right).VAlign(VAlign_Center).AutoWidth()
-		[
-			SNew(STextBlock)
-			.Text(LOCTEXT("NewAssetPathLabel", "New Asset Path"))
-			.ToolTipText(LOCTEXT("NewAssetPathToolTip", "Path used for storing newly generated assets within the project folder."))
-		]
-		+ SHorizontalBox::Slot().HAlign(HAlign_Fill).Padding(0).FillWidth(1.0f)
-		[
-			NewAssetPath->AsShared()
-		]
-		+ SHorizontalBox::Slot().HAlign(HAlign_Right).Padding(0, 0, 0, 0).AutoWidth()
+		+ SHorizontalBox::Slot()
+		.HAlign(HAlign_Fill)
+		.Padding(0, 0, 0, 8)
 		[
-			SNew(SSimpleButton)
-			.OnClicked_Lambda( [this]() { SelectNewAssetPath(); return FReply::Handled(); } )
-			.Icon(FAppStyle::Get().GetBrush("Icons.FolderOpen"))
+			SNew(SHorizontalBox)
+			+ SHorizontalBox::Slot().Padding(0).HAlign(HAlign_Left).VAlign(VAlign_Center).FillWidth(2.f)
+			[
+				SNew(STextBlock)
+				.Text(LOCTEXT("NewAssetPathLabel", "New Asset Path"))
+				.ToolTipText(LOCTEXT("NewAssetPathToolTip", "Path used for storing newly generated assets within the project folder."))
+			]
+			+ SHorizontalBox::Slot()
+			.HAlign(HAlign_Fill).
+			FillWidth(4.f)
+			[
+				SNew(SHorizontalBox)
+				+ SHorizontalBox::Slot().HAlign(HAlign_Fill).Padding(0).FillWidth(9.0f)
+					[
+						NewAssetPath->AsShared()
+					]
+					+ SHorizontalBox::Slot().HAlign(HAlign_Right).Padding(0, 0, 0, 0).AutoWidth()
+					[
+						SNew(SSimpleButton)
+						.OnClicked_Lambda( [this]() { SelectNewAssetPath(); return FReply::Handled(); } )
+						.Icon(FAppStyle::Get().GetBrush("Icons.FolderOpen"))
+					]
+			]
+
 		];
 	}
 	else
 	{
 		AssetOptionsArgs
-		+ SHorizontalBox::Slot().Padding(0, 2, 2, 2).HAlign(HAlign_Right).VAlign(VAlign_Center).AutoWidth()
-		[
-			SNew(STextBlock)
-			.Text(LOCTEXT("AssetLocationLabel", "New Asset Location"))
-		]
-		+ SHorizontalBox::Slot().HAlign(HAlign_Fill).Padding(0).FillWidth(1.0f)
+		+ SHorizontalBox::Slot()
+		.HAlign(HAlign_Fill)
+		.Padding(0, 0, 0, 0)
 		[
-			AssetLocationMode->AsShared()
+			SNew(SHorizontalBox)
+			+ SHorizontalBox::Slot().HAlign(HAlign_Left).VAlign(VAlign_Center).FillWidth(2.f)
+			[
+				SNew(STextBlock)
+				.Text(LOCTEXT("AssetLocationLabel", "New Asset Location"))
+			]
+			+ SHorizontalBox::Slot().HAlign(HAlign_Fill).FillWidth(4.f)
+			[
+				AssetLocationMode->AsShared()
+			]
 		];
 	}
-	
-	// Regardless of restrictive mode or not, we add a cog icon to the same row with some common settings.
-	AssetOptionsArgs
-	+ SHorizontalBox::Slot().HAlign(HAlign_Right).Padding(0, 0, 0, 0).AutoWidth()
-	[
-		SNew(SComboButton)
-		.HasDownArrow(false)
-		.MenuPlacement(EMenuPlacement::MenuPlacement_MenuRight)
-		.ComboButtonStyle(FAppStyle::Get(), "SimpleComboButton")
-		.OnGetMenuContent(FOnGetContent::CreateLambda([this]()
-		{
-			return MakeMenu_ModelingModeConfigSettings();
-		}))
-		.ContentPadding(FMargin(3.0f, 1.0f))
-		.ButtonContent()
-		[
-			SNew(SImage)
-			.Image(FModelingToolsEditorModeStyle::Get()->GetBrush("ModelingMode.DefaultSettings"))
-			.ColorAndOpacity(FSlateColor::UseForeground())
-		]
-	];
 
 	Content->AddSlot()[
 		SArgumentNew(AssetOptionsArgs, SHorizontalBox)
 	];
 
-
 	TSharedPtr<SExpandableArea> AssetConfigPanel = SNew(SExpandableArea)
-		.HeaderPadding(FMargin(2.0f))
-		.Padding(FMargin(2.f))
+		.HeaderPadding(FMargin(0.f))
+		.Padding(FMargin(8.f))
 		.BorderImage(FAppStyle::Get().GetBrush("DetailsView.CategoryTop"))
-		.BorderBackgroundColor(FLinearColor(.6, .6, .6, 1.0f))
-		.BodyBorderBackgroundColor(FLinearColor::Transparent)
 		.AreaTitleFont(FAppStyle::Get().GetFontStyle("EditorModesPanel.CategoryFontStyle"))
 		.BodyContent()
 		[
@@ -788,10 +760,33 @@ TSharedPtr<SWidget> FModelingToolsEditorModeToolkit::MakeAssetConfigPanel()
 		]
 		.HeaderContent()
 		[
-			SNew(STextBlock)
-			.Text(LOCTEXT("ModelingSettingsPanelHeader", "Modeling Mode Quick Settings"))
-			.Justification(ETextJustify::Center)
-			.Font(FAppStyle::Get().GetFontStyle("EditorModesPanel.CategoryFontStyle"))
+			SNew(SHorizontalBox)
+			+ SHorizontalBox::Slot().HAlign(HAlign_Left).VAlign(VAlign_Center).FillWidth(2.f)
+				[
+					SNew(STextBlock)
+					.Text(LOCTEXT("ModelingSettingsPanelHeader", "Modeling Mode Quick Settings"))
+					
+				]
+
+				+ SHorizontalBox::Slot().HAlign(HAlign_Right).VAlign(VAlign_Center).AutoWidth()
+				[
+					SNew(SComboButton)
+						.HasDownArrow(false)
+						.MenuPlacement(EMenuPlacement::MenuPlacement_MenuRight)
+						.ComboButtonStyle(FAppStyle::Get(), "SimpleComboButton")
+						.OnGetMenuContent(FOnGetContent::CreateLambda([this]()
+						{
+							return MakeMenu_ModelingModeConfigSettings();
+						}))
+						.ContentPadding(FMargin(3.0f, 1.0f))
+						.ButtonContent()
+						[
+							SNew(SImage)
+							.Image(FModelingToolsEditorModeStyle::Get()->GetBrush("ModelingMode.DefaultSettings"))
+							.ColorAndOpacity(FSlateColor::UseForeground())
+						]
+				]
+				
 		];
 
 	return AssetConfigPanel;
@@ -1861,6 +1856,11 @@ void FModelingToolsEditorModeToolkit::OnToolStarted(UInteractiveToolManager* Man
 	ModeHeaderArea->SetVisibility(EVisibility::Collapsed);
 	ActiveToolName = CurTool->GetToolInfo().ToolDisplayName;
 
+	if (HasToolkitBuilder())
+	{
+		ToolkitBuilder->SetActiveToolDisplayName(ActiveToolName);
+	}
+
 	// try to update icon
 	FString ActiveToolIdentifier = GetScriptableEditorMode()->GetToolManager(EToolsContextScope::EdMode)->GetActiveToolName(EToolSide::Left);
 	ActiveToolIdentifier.InsertAt(0, ".");
@@ -1895,6 +1895,11 @@ void FModelingToolsEditorModeToolkit::OnToolEnded(UInteractiveToolManager* Manag
 
 	ModeDetailsView->SetObject(nullptr);
 	ActiveToolName = FText::GetEmpty();
+	if (HasToolkitBuilder())
+	{
+		ToolkitBuilder->SetActiveToolDisplayName(FText::GetEmpty());		
+	}
+
 	ModeHeaderArea->SetVisibility(EVisibility::Visible);
 	ModeHeaderArea->SetText(LOCTEXT("SelectToolLabel", "Select a Tool from the Toolbar"));
 	ClearNotification();
diff --git a/Engine/Source/Developer/WidgetRegistration/Private/FToolkitWidgetStyle.cpp b/Engine/Source/Developer/WidgetRegistration/Private/FToolkitWidgetStyle.cpp
new file mode 100644
--- /dev/null
+++ b/Engine/Source/Developer/WidgetRegistration/Private/FToolkitWidgetStyle.cpp
@@ -0,0 +1,53 @@
+// Copyright Epic Games, Inc. All Rights Reserved.
+
+#include "FToolkitWidgetStyle.h"
+#include UE_INLINE_GENERATED_CPP_BY_NAME(FToolkitWidgetStyle)
+
+FToolkitWidgetStyle& FToolkitWidgetStyle::SetActiveToolTitleBorderPadding(const FMargin& InActiveToolTitleBorderPadding)
+{
+	ActiveToolTitleBorderPadding = InActiveToolTitleBorderPadding; return *this;
+}
+
+FToolkitWidgetStyle& FToolkitWidgetStyle::SetToolContextTextBlockPadding(const FMargin& InToolContextTextBlockPadding)
+{
+	ToolContextTextBlockPadding = InToolContextTextBlockPadding; return *this;
+}
+
+FToolkitWidgetStyle& FToolkitWidgetStyle::SetTitleBackgroundBrush(const FSlateBrush& InPaletteTitleBackgroundBrush)
+{
+	TitleBackgroundBrush = InPaletteTitleBackgroundBrush; return *this;
+}
+
+FToolkitWidgetStyle& FToolkitWidgetStyle::SetToolDetailsBackgroundBrush(const FSlateBrush& InToolDetailsBackgroundBrush)
+{
+	ToolDetailsBackgroundBrush = InToolDetailsBackgroundBrush; return *this;
+}
+
+const FToolkitWidgetStyle& FToolkitWidgetStyle::GetDefault()
+{
+	static FToolkitWidgetStyle Default;
+	return Default;
+}
+
+const FName FToolkitWidgetStyle::TypeName(TEXT("FToolkitWidgetStyle"));
+
+void FToolkitWidgetStyle::GetResources(TArray<const FSlateBrush*>& OutBrushes) const
+{
+	OutBrushes.Add(&TitleBackgroundBrush);
+}
+
+FToolkitWidgetStyle& FToolkitWidgetStyle::SetTitleForegroundColor(const FSlateColor& InTitleForegroundColor)
+{
+	TitleForegroundColor = InTitleForegroundColor; return *this;
+}
+
+FToolkitWidgetStyle& FToolkitWidgetStyle::SetTitlePadding(const FMargin& InTitlePadding)
+{
+	TitlePadding = InTitlePadding; return *this;
+}
+
+FToolkitWidgetStyle& FToolkitWidgetStyle::SetTitleFont(const FSlateFontInfo& InTitleFont)
+{
+	TitleFont = InTitleFont; return *this;
+}
+
diff --git a/Engine/Source/Developer/WidgetRegistration/Private/ToolkitBuilder.cpp b/Engine/Source/Developer/WidgetRegistration/Private/ToolkitBuilder.cpp
--- a/Engine/Source/Developer/WidgetRegistration/Private/ToolkitBuilder.cpp
+++ b/Engine/Source/Developer/WidgetRegistration/Private/ToolkitBuilder.cpp
@@ -2,10 +2,12 @@
 
 #include "ToolkitBuilder.h"
 
+#include "IDetailsView.h"
 #include "Widgets/SBoxPanel.h"
 #include "SPrimaryButton.h"
 #include "ToolbarRegistrationArgs.h"
 #include "ToolElementRegistry.h"
+#include "ToolkitStyle.h"
 #include "Framework/Commands/UICommandList.h"
 #include "Styling/StyleColors.h"
 
@@ -31,10 +33,12 @@ FEditablePalette::FEditablePalette(TSharedPtr<FUICommandInfo> InLoadToolPaletteA
 
 FToolkitBuilder::FToolkitBuilder(
 	FName InToolbarCustomizationName,
-	TSharedPtr<FUICommandList> InToolkitCommandList) :
+	TSharedPtr<FUICommandList> InToolkitCommandList,
+	TSharedPtr<FToolkitSections> InToolkitSections) :
 	FToolElementRegistrationArgs(EToolElement::Toolkit),
 	ToolbarCustomizationName(InToolbarCustomizationName),
-	ToolkitCommandList(InToolkitCommandList)
+	ToolkitCommandList(InToolkitCommandList),
+	ToolkitSections(InToolkitSections)
 {
 	ResetWidget();
 }
@@ -55,13 +59,13 @@ void FToolkitBuilder::AddPalette(TSharedPtr<FEditablePalette> Palette)
 	AddPalette(StaticCastSharedRef<FToolPalette>(Palette.ToSharedRef()) );
 }
 
-
 void FToolkitBuilder::AddPalette(TSharedPtr<FToolPalette> Palette)
 {
 	for (TSharedRef<FButtonArgs> Button : Palette->PaletteActions)
 	{
 		PaletteCommandNameToButtonArgsMap.Add(Button->Command->GetCommandName().ToString(), Button);
 	}
+	LoadCommandNameToToolPaletteMap.Add(Palette->LoadToolPaletteAction->GetCommandName().ToString(), Palette);
 
 	LoadToolPaletteCommandList->MapAction(
 				Palette->LoadToolPaletteAction,
@@ -139,10 +143,12 @@ void FToolkitBuilder::TogglePalette(TSharedPtr<FToolPalette> Palette)
 	
 	if (ActivePalette && ActivePalette->LoadToolPaletteAction->GetCommandName() == CommandName)
 	{
-		ActivePalette = nullptr;
-		ResetToolPaletteWidget();
-	    OnToolEnded.ExecuteIfBound();
 		return;
+/*
+ *		@TODO: ~ enable this code for hiding category select on toggle		
+ *		ActivePalette = nullptr;
+		ResetToolPaletteWidget();
+		return; */
 	}
 	
 	CreatePalette(Palette);
@@ -150,6 +156,11 @@ void FToolkitBuilder::TogglePalette(TSharedPtr<FToolPalette> Palette)
 
 void FToolkitBuilder::CreatePalette(TSharedPtr<FToolPalette> Palette)
 {
+	if (!Palette)
+	{
+		return;
+	}
+	
 	const FName CommandName = Palette->LoadToolPaletteAction->GetCommandName();
 	ActivePalette = Palette;
 	ResetToolPaletteWidget();
@@ -179,13 +190,19 @@ void FToolkitBuilder::CreatePalette(TSharedPtr<FToolPalette> Palette)
 	for (TSharedRef<FButtonArgs> PaletteButton : Palette->PaletteActions)
 	{
 		PaletteButton->CommandList = ToolkitCommandList;
-		PaletteButton->UserInterfaceActionType = EUserInterfaceActionType::ToggleButton;
+		PaletteButton->UserInterfaceActionType = (PaletteButton->UserInterfaceActionType != EUserInterfaceActionType::None) ?
+														PaletteButton->UserInterfaceActionType :
+														EUserInterfaceActionType::ToggleButton;
 		PaletteButton->OnGetMenuContent.BindSP(SharedThis(this),
 			&FToolkitBuilder::GetContextMenuContent,
 			PaletteButton->Command->GetCommandName());
 		PaletteToolbarBuilder->AddToolBarButton(PaletteButton.Get());
 	}
+	CreatePaletteWidget(*Palette.Get(), *Element.Get());
+}
 
+void FToolkitBuilder::CreatePaletteWidget(FToolPalette& Palette, FToolElement& Element)
+{
 	ToolPaletteWidget->AddSlot()
 		.HAlign(HAlign_Fill)
 		.Padding(0.f)
@@ -193,29 +210,28 @@ void FToolkitBuilder::CreatePalette(TSharedPtr<FToolPalette> Palette)
 		[
 
 		SNew(SBorder)
-		.Padding(FAppStyle::GetMargin("SlimPaletteToolBarStyle.TitleBar.Padding"))
+		.Padding(Style.TitlePadding)
 		.VAlign(VAlign_Center)
-		.BorderImage(FAppStyle::GetBrush("ToolPanel.GroupBorder"))
+		.BorderImage(&Style.TitleBackgroundBrush)
 		.HAlign(HAlign_Left)
 		[
 
 			SNew(STextBlock)
 					.Justification(ETextJustify::Left)
-					.AutoWrapText(true)
-					.Text(Palette->LoadToolPaletteAction->GetLabel())
-					.ColorAndOpacity(FStyleColors::Foreground)
+					.Font(Style.TitleFont)
+					.Text(Palette.LoadToolPaletteAction->GetLabel())
+					.ColorAndOpacity(Style.TitleForegroundColor)
 		]];
 	
 	ToolPaletteWidget->AddSlot()
 		.HAlign(HAlign_Fill)
 		.AutoHeight()
 		[
-			Element->GenerateWidget()
+			Element.GenerateWidget()
 		];
-
-	OnToolEnded.ExecuteIfBound();
 }
 
+
 TSharedRef<SWidget> FToolkitBuilder::GetToolPaletteWidget() const
 {
 	return ToolPaletteWidget->AsShared();
@@ -245,6 +261,7 @@ TSharedRef<SWidget> FToolkitBuilder::GetContextMenuContent(const FName CommandNa
 
 void FToolkitBuilder::ResetWidget()
 {
+	Style = FToolkitStyle::Get().GetWidgetStyle<FToolkitWidgetStyle>("FToolkitWidgetStyle");
 	LoadToolPaletteCommandList = MakeShareable(new FUICommandList);
 	LoadPaletteToolBarBuilder = MakeShared<FVerticalToolBarBuilder>(LoadToolPaletteCommandList, FMultiBoxCustomization::None, TSharedPtr<FExtender>(), true);
 	ToolPaletteWidget = SNew(SVerticalBox);
@@ -265,7 +282,7 @@ void FToolkitBuilder::ResetWidget()
 		ToolRegistry.RegisterElement(VerticalToolbarElement.ToSharedRef());
 	}
 
-	VerticalToolbarElement->SetRegistrationArgs(VerticalToolbarRegistrationArgs);	
+	VerticalToolbarElement->SetRegistrationArgs(VerticalToolbarRegistrationArgs);
 }
 
 void FToolkitBuilder::ResetToolPaletteWidget()
@@ -278,4 +295,144 @@ void FToolkitBuilder::ResetToolPaletteWidget()
 	ToolPaletteWidget = SNew(SVerticalBox);
 }
 
+bool FToolkitBuilder::HasSelectedToolSet() const
+{
+	return HasActivePalette();
+}
+
+void FToolkitBuilder::SetActivePaletteOnLoad(const FUICommandInfo* Command)
+{
+	if (const TSharedPtr<FToolPalette>* LoadPalette = LoadCommandNameToToolPaletteMap.Find(Command->GetCommandName().ToString()))
+	{
+		CreatePalette(*LoadPalette);
+	}
+}
+
+TSharedPtr<SWidget> FToolkitBuilder::GenerateWidget()
+{
+	if (!ToolkitWidgetHBox)
+	{
+		DefineWidget();
+	}
+	return ToolkitWidgetHBox.ToSharedRef();
+}
+
+void FToolkitBuilder::SetActiveToolDisplayName(FText InActiveToolDisplayName)
+{
+	ActiveToolDisplayName = InActiveToolDisplayName;
+}
+
+FText FToolkitBuilder::GetActiveToolDisplayName() const
+{
+	return ActiveToolDisplayName;
+}
+
+EVisibility FToolkitBuilder::GetActiveToolTitleVisibility() const
+{
+	return ActiveToolDisplayName.IsEmpty() ? EVisibility::Collapsed : EVisibility::Visible;
+}
+
+void FToolkitBuilder::DefineWidget()
+{
+	ToolkitWidgetVBox = SNew(SVerticalBox);
+	ToolkitWidgetHBox = 
+		SNew(SSplitter)
+		+ SSplitter::Slot()
+		.Resizable(false)
+		.SizeRule(SSplitter::SizeToContent)
+			[
+				CreateToolbarWidget()
+			]
+
+		+ SSplitter::Slot()
+		.SizeRule(SSplitter::FractionOfParent)
+			[
+				ToolkitWidgetVBox->AsShared()
+			];
+
+	if (ToolkitSections->ModeWarningArea)
+	{
+		ToolkitWidgetVBox->AddSlot()
+		.AutoHeight()
+		.HAlign(HAlign_Fill)
+		.Padding(5)
+		[
+			ToolkitSections->ModeWarningArea->AsShared()
+		];
+	}
+	
+	ToolkitWidgetVBox->AddSlot()
+	.AutoHeight()
+	.HAlign(HAlign_Fill)
+	.Padding(0)
+		[
+			GetToolPaletteWidget()
+		];
+
+	ToolkitWidgetVBox->AddSlot()
+		.AutoHeight()
+		.HAlign(HAlign_Fill)
+		.Padding(0)
+		[
+			SNew(SBorder)
+			.HAlign(HAlign_Fill)
+			.Padding(Style.ActiveToolTitleBorderPadding)
+			.BorderImage(&Style.ToolDetailsBackgroundBrush)
+			[
+			SNew(SBorder)
+			.Visibility(SharedThis(this), &FToolkitBuilder::GetActiveToolTitleVisibility)
+			.BorderImage(&Style.TitleBackgroundBrush)
+			.Padding(Style.ToolContextTextBlockPadding)
+			.HAlign(HAlign_Left)
+			.VAlign(VAlign_Top)
+			[
+
+				SNew(STextBlock)
+						.Justification(ETextJustify::Left)
+						.Font(Style.TitleFont)
+						.Text(TAttribute<FText>(SharedThis(this), &FToolkitBuilder::GetActiveToolDisplayName))
+						.ColorAndOpacity(Style.TitleForegroundColor)
+			]
+		]
+	];
+	
+	if (ToolkitSections->ToolWarningArea)
+	{
+		ToolkitWidgetVBox->AddSlot()
+		.AutoHeight()
+		.HAlign(HAlign_Fill)
+		.Padding(5)
+		[
+			ToolkitSections->ToolWarningArea->AsShared()
+		];
+	}
+
+	if (ToolkitSections->DetailsView)
+	{
+		ToolkitWidgetVBox->AddSlot()
+		.HAlign(HAlign_Fill)
+		.FillHeight(1.f)
+			[
+			SNew(SBorder)
+			.BorderImage(&Style.ToolDetailsBackgroundBrush)
+			.Padding(8.f, 2.f, 0.f, 2.f)
+				[
+					ToolkitSections->DetailsView->AsShared()
+				]
+		
+			];
+	}
+	if (ToolkitSections->Footer)
+	{
+		ToolkitWidgetVBox->AddSlot()
+			.AutoHeight()
+			.HAlign(HAlign_Fill)
+			.VAlign(VAlign_Bottom)
+			.Padding(0)
+		[
+			ToolkitSections->Footer->AsShared()
+		];
+	}
+}
+
 #undef LOCTEXT_NAMESPACE
\ No newline at end of file
diff --git a/Engine/Source/Developer/WidgetRegistration/Private/ToolkitStyle.cpp b/Engine/Source/Developer/WidgetRegistration/Private/ToolkitStyle.cpp
new file mode 100644
--- /dev/null
+++ b/Engine/Source/Developer/WidgetRegistration/Private/ToolkitStyle.cpp
@@ -0,0 +1,68 @@
+// Copyright Epic Games, Inc. All Rights Reserved.
+
+#include "ToolkitStyle.h"
+
+#include "FToolkitWidgetStyle.h"
+#include "Styling/SlateStyleRegistry.h"
+#include "Interfaces/IPluginManager.h"
+#include "Styling/SlateStyle.h"
+#include "Styling/StarshipCoreStyle.h"
+#include "Styling/StyleColors.h"
+
+TSharedPtr< FSlateStyleSet > FToolkitStyle::StyleSet = nullptr;
+FName FToolkitStyle::StyleName("ToolkitStyle");
+
+FToolkitStyle::FToolkitStyle()
+	: FSlateStyleSet(StyleName)
+{
+	
+}
+
+void FToolkitStyle::Initialize()
+{
+	if (!StyleSet.IsValid())
+	{
+		StyleSet = Create();
+		FSlateStyleRegistry::RegisterSlateStyle(*StyleSet);
+	}
+
+	
+	FToolkitWidgetStyle Style = FStarshipCoreStyle::GetCoreStyle().GetWidgetStyle<FToolkitWidgetStyle>("FToolkitWidget");
+
+	Style.SetTitleBackgroundBrush(FSlateRoundedBoxBrush(FStyleColors::Header, 0.f, FLinearColor(0, 0, 0, .8), 0.5))
+		 .SetToolDetailsBackgroundBrush( FSlateRoundedBoxBrush(FStyleColors::Background, 0.f, FLinearColor(0, 0, 0, .8), 0.5))
+		 .SetTitleForegroundColor(FStyleColors::Foreground)
+		 .SetTitlePadding(FMargin(16.f, 4.f))
+	 	 .SetTitleFont(FCoreStyle::GetDefaultFontStyle("Bold", 10))
+		 .SetActiveToolTitleBorderPadding(FMargin(0.f, 8.f, 0.f, 0.f))
+		 .SetToolContextTextBlockPadding(FMargin(16.f, 4.f));
+
+	StyleSet->Set("FToolkitWidgetStyle", Style);
+}
+
+void FToolkitStyle::Shutdown()
+{
+	FSlateStyleRegistry::UnRegisterSlateStyle(*StyleSet);
+	ensure(StyleSet.IsUnique());
+	StyleSet.Reset();
+}
+
+const FName& FToolkitStyle::GetStyleSetName() const
+{
+	return StyleName;
+}
+
+TSharedRef< FSlateStyleSet > FToolkitStyle::Create()
+{
+	TSharedRef< FSlateStyleSet > Style = MakeShareable(new FSlateStyleSet(StyleName));
+
+	return Style;
+}
+
+const ISlateStyle& FToolkitStyle::Get()
+{
+	return *StyleSet;
+}
+
+#undef LOCAL_IMAGE_BRUSH
+#undef IMAGE_PLUGIN_BRUSH
diff --git a/Engine/Source/Developer/WidgetRegistration/Private/WidgetRegistrationModule.cpp b/Engine/Source/Developer/WidgetRegistration/Private/WidgetRegistrationModule.cpp
--- a/Engine/Source/Developer/WidgetRegistration/Private/WidgetRegistrationModule.cpp
+++ b/Engine/Source/Developer/WidgetRegistration/Private/WidgetRegistrationModule.cpp
@@ -1,17 +1,17 @@
 // Copyright Epic Games, Inc. All Rights Reserved.
 
-// #include "ToolkitStyle.h"
 #include "WidgetRegistrationModule.h"
+#include "ToolkitStyle.h"
 
 
-// @TODO: ~enable ToolkitStyle set 
-
 void FWidgetRegistrationModule::StartupModule()
 {
+	FToolkitStyle::Initialize();
 }
 	
 void FWidgetRegistrationModule::ShutdownModule()
 {
+	FToolkitStyle::Shutdown();
 }
 
 
diff --git a/Engine/Source/Developer/WidgetRegistration/Public/FToolkitWidgetStyle.h b/Engine/Source/Developer/WidgetRegistration/Public/FToolkitWidgetStyle.h
new file mode 100644
--- /dev/null
+++ b/Engine/Source/Developer/WidgetRegistration/Public/FToolkitWidgetStyle.h
@@ -0,0 +1,60 @@
+// Copyright Epic Games, Inc. All Rights Reserved.
+
+#pragma once
+
+#include "CoreMinimal.h"
+#include "Styling/StyleColors.h"
+#include "Brushes/SlateNoResource.h"
+
+#include "FToolkitWidgetStyle.generated.h"
+
+/**
+ * FToolkitWidgetStyle is the FSlateWidgetStyle that defines the styling of a ToolkitWidget 
+ */
+USTRUCT(BlueprintType)
+struct WIDGETREGISTRATION_API FToolkitWidgetStyle : public FSlateWidgetStyle
+{
+	GENERATED_BODY()
+
+	FToolkitWidgetStyle() :
+		TitleBackgroundBrush(FSlateNoResource()),
+		ToolDetailsBackgroundBrush(FSlateNoResource()),
+		TitleForegroundColor(FStyleColors::Panel)
+	{
+		
+	}
+
+	static const FName TypeName;
+	virtual const FName GetTypeName() const override { return TypeName; }
+
+	static const FToolkitWidgetStyle& GetDefault();
+		virtual void GetResources(TArray<const FSlateBrush*>& OutBrushes) const override;
+
+	UPROPERTY(EditAnywhere, BlueprintReadWrite, Category = Appearance)
+	FSlateBrush TitleBackgroundBrush;
+	FToolkitWidgetStyle& SetTitleBackgroundBrush(const FSlateBrush& InPaletteTitleBackgroundBrush);
+
+	UPROPERTY(EditAnywhere, BlueprintReadWrite, Category = Appearance)
+	FSlateBrush ToolDetailsBackgroundBrush;
+	FToolkitWidgetStyle& SetToolDetailsBackgroundBrush(const FSlateBrush& InToolDetailsBackgroundBrush);
+
+	UPROPERTY(EditAnywhere, BlueprintReadWrite, Category = Appearance)
+	FSlateColor TitleForegroundColor;
+	FToolkitWidgetStyle& SetTitleForegroundColor(const FSlateColor& InTitleForegroundColor);
+	
+	UPROPERTY(EditAnywhere, BlueprintReadWrite, Category = Appearance)
+	FMargin TitlePadding;
+	FToolkitWidgetStyle& SetTitlePadding(const FMargin& InTitlePadding);
+	
+	UPROPERTY(EditAnywhere, BlueprintReadWrite, Category = Appearance)
+	FMargin ActiveToolTitleBorderPadding;
+	FToolkitWidgetStyle& SetActiveToolTitleBorderPadding(const FMargin& InActiveToolTitleBorderPadding);
+
+	UPROPERTY(EditAnywhere, BlueprintReadWrite, Category = Appearance)
+	FMargin ToolContextTextBlockPadding;
+	FToolkitWidgetStyle& SetToolContextTextBlockPadding(const FMargin& InToolContextTextBlockPadding);
+
+	UPROPERTY(EditAnywhere, BlueprintReadWrite, Category = Appearance)
+	FSlateFontInfo TitleFont;
+	FToolkitWidgetStyle& SetTitleFont(const FSlateFontInfo& InTitleFont);
+};
\ No newline at end of file
diff --git a/Engine/Source/Developer/WidgetRegistration/Public/ToolkitBuilder.h b/Engine/Source/Developer/WidgetRegistration/Public/ToolkitBuilder.h
--- a/Engine/Source/Developer/WidgetRegistration/Public/ToolkitBuilder.h
+++ b/Engine/Source/Developer/WidgetRegistration/Public/ToolkitBuilder.h
@@ -2,17 +2,28 @@
 
 #pragma once
 
+#include "FToolkitWidgetStyle.h"
+#include "IDetailsView.h"
 #include "Framework/MultiBox/SToolBarButtonBlock.h"
 #include "Framework/Commands/UICommandInfo.h"
 #include "InteractiveToolManager.h"
 #include "ToolElementRegistry.h"
 #include "ToolMenus.h"
 #include "Widgets/SBoxPanel.h"
+#include "Widgets/Layout/SSplitter.h"
 
 
 class SWidget;
 class FToolElement;
 
+struct WIDGETREGISTRATION_API FToolkitSections
+{
+	TSharedPtr<STextBlock> ModeWarningArea = nullptr;
+	TSharedPtr<STextBlock> ToolWarningArea = nullptr;
+	TSharedPtr<IDetailsView> DetailsView = nullptr;
+	TSharedPtr<SWidget> Footer = nullptr;
+};
+
 /** A struct that provides the data for a single tool Palette*/
 struct WIDGETREGISTRATION_API FToolPalette : TSharedFromThis<FToolPalette>
 {
@@ -72,11 +83,8 @@ class WIDGETREGISTRATION_API FToolkitBuilder : public FToolElementRegistrationAr
 public:
 	FToolkitBuilder(
 		FName InToolbarCustomizationName,
-		TSharedPtr<FUICommandList> InToolkitCommandList);
-
-	/** A delegate for any clean up that may be needed when the load tool palette FUICommandAction ends/unloads */
-	DECLARE_DELEGATE(FOnToolEnded);
-	FOnToolEnded OnToolEnded;
+		TSharedPtr<FUICommandList> InToolkitCommandList,
+		TSharedPtr<FToolkitSections> InToolkitSections);
 
 	/**
 	 * Adds the FToolPalette Palette to this FToolkitBuilder
@@ -99,6 +107,12 @@ public:
 	 * will be regenerated to reflect it.  */
 	virtual void UpdateWidget() override;
 
+	/** Implements the generation of the TSharedPtr<SWidget> */
+	virtual TSharedPtr<SWidget> GenerateWidget() override;
+
+	/* Returns true if the Toolkit builder has some tools that are currently active/selected */
+	bool HasSelectedToolSet() const;
+
 	/** Creates the Toolbar for the widget with the FUICommandInfos that load the Palettes */
 	TSharedRef<SWidget> CreateToolbarWidget() const;
 
@@ -114,6 +128,13 @@ public:
 	/** returns true is there is an active palette selected, else it returns false */
 	bool HasActivePalette() const;
 
+	/*
+	 * Loads the palette for the FUICommandInfo Command on first visit to the mode.
+	 *
+	 *  @param Command the FUICommandInfo which defines the palette that will be loaded on first visit to the mode.
+	 */
+	void SetActivePaletteOnLoad(const FUICommandInfo* Command);
+	
 	/**
 	 * Returns true if the FUICommandInfo with the name CommandName is the active tool palette,
 	 * else it returns false
@@ -121,16 +142,29 @@ public:
 	 * @param CommandName the name of the FUICommandInfo we are checking to see if it is the active tool palette
 	 */
 	ECheckBoxState IsActiveToolPalette(FName CommandName) const;
+
+	void SetActiveToolDisplayName(FText InActiveToolDisplayName);
+
+	FText GetActiveToolDisplayName() const;
+
 	
 private:
 
 	/** the tool element registry this class will use to register UI tool elements */
 	static FToolElementRegistry ToolRegistry;
+
+	/* The SWidget that is the whole Toolkit */
+	TSharedPtr<SWidget> ToolkitWidget;
 	
 	/** Name of the toolbar this mode uses and can be used by external systems to customize that mode toolbar */
 	FName ToolbarCustomizationName;
+
+	/** The map of the command name to the ButtonArgs for it  */
 	TMap<FString, TSharedPtr<FButtonArgs>> PaletteCommandNameToButtonArgsMap;
 
+	/** The map of the command name to the ButtonArgs for it  */
+	TMap<FString, TSharedPtr<FToolPalette>> LoadCommandNameToToolPaletteMap;
+
 	/** A TSharedPointer to the FUICommandList for the FUICommandInfos which load a tool palette */
 	TSharedPtr<FUICommandList> LoadToolPaletteCommandList;
 
@@ -189,4 +223,31 @@ private:
 	 * @return the TSharedRef<SWidget> which contains the context menu for the FUICommandInfo with the name CommandName 
 	 */
 	TSharedRef<SWidget> GetContextMenuContent(const FName CommandName);
+
+	/**
+	 * Gets the EVisibility of the active tool title. This should only be visible if a tool is currently chosen in the palette.
+	 *
+	 * @return the EVisibility of the active tool title
+	 */
+	EVisibility GetActiveToolTitleVisibility() const;
+
+	void CreatePaletteWidget(FToolPalette& Palette, FToolElement& Element);
+
+	void DefineWidget();
+
+	/** The SVerticalBox which holds all but the vertical toolbar in a Toolkit */
+	TSharedPtr<SVerticalBox> ToolkitWidgetVBox;
+
+	/** The SSplitter which holds the entire Toolkit */
+	TSharedPtr<SSplitter> ToolkitWidgetHBox;
+
+	/** The FToolkitSections which holds the sections defined for this Toolkit */
+	TSharedPtr<FToolkitSections> ToolkitSections;
+
+	/** The display name of the currently active tool */
+	FText ActiveToolDisplayName = FText::GetEmpty();
+
+	/** The current FToolkitWidgetStyle */
+	FToolkitWidgetStyle Style;
+
 };
\ No newline at end of file
diff --git a/Engine/Source/Developer/WidgetRegistration/Public/ToolkitStyle.h b/Engine/Source/Developer/WidgetRegistration/Public/ToolkitStyle.h
new file mode 100644
--- /dev/null
+++ b/Engine/Source/Developer/WidgetRegistration/Public/ToolkitStyle.h
@@ -0,0 +1,28 @@
+// Copyright Epic Games, Inc. All Rights Reserved.
+
+#pragma once
+
+#include "Styling/SlateStyle.h"
+#include "Templates/SharedPointer.h"
+
+class ISlateStyle;
+
+/** FToolkitStyle is the FSlateStyleSet that defines styles for FToolkitBuilders  */
+class WIDGETREGISTRATION_API FToolkitStyle
+	: public FSlateStyleSet
+{
+public:
+	
+	static FName StyleName;
+	FToolkitStyle();
+	static void Initialize();
+	static void Shutdown();
+	static const ISlateStyle& Get();
+	virtual const FName& GetStyleSetName() const override;
+
+private:
+
+	static TSharedRef< class FSlateStyleSet > Create();
+	static TSharedPtr< class FSlateStyleSet > StyleSet;
+};
+
diff --git a/Engine/Source/Developer/WidgetRegistration/WidgetRegistration.Build.cs b/Engine/Source/Developer/WidgetRegistration/WidgetRegistration.Build.cs
--- a/Engine/Source/Developer/WidgetRegistration/WidgetRegistration.Build.cs
+++ b/Engine/Source/Developer/WidgetRegistration/WidgetRegistration.Build.cs
@@ -13,7 +13,8 @@ public class WidgetRegistration : ModuleRules
 				"Core",
 				"Slate",
 				"Engine",
-				"ToolMenus"
+				"ToolMenus",
+				"CoreUObject"
 			});
 		
 
diff --git a/Engine/Source/Editor/EditorStyle/Private/StarshipStyle.cpp b/Engine/Source/Editor/EditorStyle/Private/StarshipStyle.cpp
--- a/Engine/Source/Editor/EditorStyle/Private/StarshipStyle.cpp
+++ b/Engine/Source/Editor/EditorStyle/Private/StarshipStyle.cpp
@@ -1682,39 +1682,47 @@ void FStarshipEditorStyle::FStyle::SetupGeneralStyles()
 		FTextBlockStyle SlimPaletteToolbarLabelStyle = FTextBlockStyle(GetParentStyle()->GetWidgetStyle<FTextBlockStyle>("NormalText"));
 		SlimPaletteToolbarLabelStyle.SetOverflowPolicy(ETextOverflowPolicy::Ellipsis);
 		SlimPaletteToolBarStyle.SetLabelStyle(SlimPaletteToolbarLabelStyle);
-		SlimPaletteToolBarStyle.SetBackground(FSlateColorBrush(FStyleColors::Recessed));
+		SlimPaletteToolBarStyle.SetBackgroundPadding(FMargin(2.0f, 4.0f, 6.0f, 8.0f));
+		SlimPaletteToolBarStyle.SetBackground(FSlateColorBrush(FStyleColors::Panel));
 		SlimPaletteToolBarStyle.SetLabelPadding(FMargin(0.0f, 0.0f, 0.0f, 0.0f));
 		SlimPaletteToolBarStyle.SetIconSize(Icon20x20);
 
-		SlimPaletteToolBarStyle.SetButtonPadding(       FMargin(0.0f, 4.0f, 0.0f, 0.0f));
-		SlimPaletteToolBarStyle.SetCheckBoxPadding(     FMargin(0.0f, 0.0f));
-		// SlimPaletteToolBarStyle.SetIconPadding( FMargin(8.f, 4.f, 8.f, 4.f) ); @TODO: ~ enable with styling changesS
+		SlimPaletteToolBarStyle.SetButtonPadding(       FMargin(4.0f, 4.0f, 0.0f, 0.0f));
+		SlimPaletteToolBarStyle.SetCheckBoxPadding(     FMargin(4.0f, 4.0f, 0.0f, 0.0f));
+		SlimPaletteToolBarStyle.SetIconPadding( FMargin(8.f, 4.f, 8.f, 4.f) );
 		SlimPaletteToolBarStyle.SetComboButtonPadding(  FMargin(0.0f, 0.0f));
 		SlimPaletteToolBarStyle.SetIndentedBlockPadding(FMargin(0.0f, 0.0f));
 		SlimPaletteToolBarStyle.SetBlockPadding(        FMargin(0.0f, 0.0f));
-		SlimPaletteToolBarStyle.ToggleButton.SetPadding(FMargin(0.0f, 1.0f));
-		SlimPaletteToolBarStyle.ButtonStyle.SetNormalPadding(FMargin(2.0f, 1.0f));
-		SlimPaletteToolBarStyle.ButtonStyle.SetPressedPadding(FMargin(2.0f, 1.0f));
-
+		SlimPaletteToolBarStyle.ToggleButton.SetPadding(FMargin(0.0f, 0.0f));
+		SlimPaletteToolBarStyle.ButtonStyle.SetNormalPadding(FMargin(0.0f, 0.0f));
+		SlimPaletteToolBarStyle.ButtonStyle.SetPressedPadding(FMargin(0.0f, 0.0f));
+
+		SlimPaletteToolBarStyle.ButtonStyle.Normal = FSlateRoundedBoxBrush(FStyleColors::Dropdown, 4.f, FLinearColor(0, 0, 0, .8), 0.5);
+		SlimPaletteToolBarStyle.ButtonStyle.Hovered = FSlateRoundedBoxBrush(FStyleColors::Hover, 4.f, FLinearColor(0, 0, 0, .8), 0.5);
+		SlimPaletteToolBarStyle.ButtonStyle.Pressed = FSlateRoundedBoxBrush(FStyleColors::Hover, 4.f, FLinearColor(0, 0, 0, .8), 0.5);
+		SlimPaletteToolBarStyle.ButtonStyle.HoveredForeground = FStyleColors::ForegroundHover;
+		SlimPaletteToolBarStyle.ButtonStyle.PressedForeground = FStyleColors::ForegroundHover;
+		SlimPaletteToolBarStyle.SetUniformBlockHeight(33.f);
+		SlimPaletteToolBarStyle.SetUniformBlockWidth(150.f);
+		SlimPaletteToolBarStyle.SetNumColumns(2);
+		
 		FCheckBoxStyle CheckBoxStyle = FCheckBoxStyle(SlimPaletteToolBarStyle.ToggleButton)
 			.SetCheckedImage(FSlateRoundedBoxBrush(FStyleColors::Primary, 4.f, FLinearColor(0, 0, 0, .8), 0.5))
 			.SetCheckedHoveredImage(FSlateRoundedBoxBrush(FStyleColors::PrimaryHover, 4.f, FLinearColor(0, 0, 0, .8), 0.5))
-			.SetCheckedPressedImage(FSlateRoundedBoxBrush(FStyleColors::PrimaryPress, 4.f, FLinearColor(0, 0, 0, .8), 0.5))
+			.SetCheckedPressedImage(FSlateRoundedBoxBrush(FStyleColors::PrimaryHover, 4.f, FLinearColor(0, 0, 0, .8), 0.5))
 			.SetUncheckedImage(FSlateRoundedBoxBrush(FStyleColors::Dropdown, 4.f, FLinearColor(0, 0, 0, .8), 0.5))
 			.SetForegroundColor(FStyleColors::Foreground)
 			.SetPressedForegroundColor(FStyleColors::ForegroundHover)
 			.SetHoveredForegroundColor(FStyleColors::ForegroundHover)
-			.SetCheckedForegroundColor(FStyleColors::Foreground)
+			.SetCheckedForegroundColor(FStyleColors::ForegroundHover)
 			.SetCheckedPressedForegroundColor(FStyleColors::ForegroundHover)
 			.SetCheckedHoveredForegroundColor(FStyleColors::ForegroundHover)
 			.SetPadding(FMargin(0.f, 0.f, 0.f, 0.f));
 
 		Set("SlimPaletteToolBarStyle.ToggleButton", CheckBoxStyle);
-		Set("SlimPaletteToolBarStyle.TitleBar.Padding", FMargin(4.f, 4.f));
 		SlimPaletteToolBarStyle.SetToggleButtonStyle(CheckBoxStyle);
 		Set("SlimPaletteToolBar", SlimPaletteToolBarStyle);
-		Set("ToolAreaBackground",  new FSlateRoundedBoxBrush(FStyleColors::AccentGreen, 4.0f, FLinearColor(0, 0, 0, .8), 0.5));
-		Set( "SlimPaletteToolBar.PaletteButtonBorderBrush", new FSlateRoundedBoxBrush(FStyleColors::AccentGreen, 4.0f, FLinearColor(0, 0, 0, .8), 0.5) );
+		
 	}
 
 {
@@ -1724,17 +1732,18 @@ void FStarshipEditorStyle::FStyle::SetupGeneralStyles()
 		FVerticalToolBarLabelStyle.SetOverflowPolicy(ETextOverflowPolicy::Ellipsis);
 		
 		FVerticalToolBarStyle.SetLabelStyle(FVerticalToolBarLabelStyle);
-		FVerticalToolBarStyle.SetLabelPadding(FMargin(0.0f, 1.0f, 0.0f, 0.0f));
+		FVerticalToolBarStyle.SetLabelPadding(FMargin(0.0f, 4.0f, 0.0f, 0.0f));
 
-		FVerticalToolBarStyle.SetButtonPadding(       FMargin(0.0f, 0.0f));
+		FVerticalToolBarStyle.SetButtonPadding(       FMargin(0.0f, 0.0f, 0.0f, 0.0f));
+		//not this
 		FVerticalToolBarStyle.SetCheckBoxPadding(     FMargin(0.0f, 0.0f));
 		FVerticalToolBarStyle.SetComboButtonPadding(  FMargin(0.0f, 0.0f));
-		FVerticalToolBarStyle.SetIndentedBlockPadding(FMargin(0.0f, 0.0f));
-		FVerticalToolBarStyle.SetBlockPadding(        FMargin(0.0f, 0.0f));
-		FVerticalToolBarStyle.SetBackgroundPadding(   FMargin(4.0f, 2.0f));
-		FVerticalToolBarStyle.ButtonStyle.SetNormalPadding(FMargin(2.0f, 6.0f));
-		FVerticalToolBarStyle.ButtonStyle.SetPressedPadding(FMargin(2.0f, 6.0f));
+		FVerticalToolBarStyle.SetIndentedBlockPadding(FMargin(0.0f, 0.0f));	
+		FVerticalToolBarStyle.SetBackgroundPadding(   FMargin(6.0f, 4.0f));
+		FVerticalToolBarStyle.ButtonStyle.SetNormalPadding(FMargin(12.0f, 6.0f));
+		FVerticalToolBarStyle.ButtonStyle.SetPressedPadding(FMargin(12.0f, 6.0f));
 		FVerticalToolBarStyle.SetExpandBrush(IMAGE_BRUSH("Icons/toolbar_expand_16x", Icon8x8));
+		FVerticalToolBarStyle.SetIconPadding(FMargin(8.0f, 8.0f));
  
 		FCheckBoxStyle CheckBoxStyle = FCheckBoxStyle(FVerticalToolBarStyle.ToggleButton)
 			.SetCheckedImage(FSlateRoundedBoxBrush(FStyleColors::Primary, 4.f, FLinearColor(0, 0, 0, .8), 0.5))
@@ -1743,17 +1752,15 @@ void FStarshipEditorStyle::FStyle::SetupGeneralStyles()
 			.SetForegroundColor(FStyleColors::Foreground)
 			.SetPressedForegroundColor(FStyleColors::ForegroundHover)
 			.SetHoveredForegroundColor(FStyleColors::ForegroundHover)
-			.SetCheckedForegroundColor(FStyleColors::Foreground)
+			.SetCheckedForegroundColor(FStyleColors::ForegroundHover)
 			.SetCheckedPressedForegroundColor(FStyleColors::ForegroundHover)
 			.SetCheckedHoveredForegroundColor(FStyleColors::ForegroundHover)
-			.SetPadding(FMargin(6.0f, 4.0f, 6.0f, 4.0f));
+			.SetPadding(0.f);
 
 		Set("FVerticalToolBar.ToggleButton", CheckBoxStyle);
-		FVerticalToolBarStyle.SetToggleButtonStyle(CheckBoxStyle);
-
-		Set("FVerticalToolBar.MaxUniformToolbarSize", 48.f);
-		Set("FVerticalToolBar.MinUniformToolbarSize", 48.f);
 
+		FVerticalToolBarStyle.SetButtonPadding(FMargin(0.0f, 4.0f, 0.0f, 0.f));
+		FVerticalToolBarStyle.SetToggleButtonStyle(CheckBoxStyle);
 		Set("FVerticalToolBar", FVerticalToolBarStyle);
 	}
 	
diff --git a/Engine/Source/Editor/UnrealEd/Private/EditorEngine.cpp b/Engine/Source/Editor/UnrealEd/Private/EditorEngine.cpp
--- a/Engine/Source/Editor/UnrealEd/Private/EditorEngine.cpp
+++ b/Engine/Source/Editor/UnrealEd/Private/EditorEngine.cpp
@@ -1211,6 +1211,7 @@ void UEditorEngine::Init(IEngineLoop* InEngineLoop)
 		}
 
 		FModuleManager::Get().LoadModule(TEXT("LogVisualizer"));
+		FModuleManager::Get().LoadModule(TEXT("WidgetRegistration"));
 		FModuleManager::Get().LoadModule(TEXT("HotReload"));
 
 		FModuleManager::Get().LoadModuleChecked(TEXT("ClothPainter"));
diff --git a/Engine/Source/Editor/UnrealEd/Public/Toolkits/BaseToolkit.h b/Engine/Source/Editor/UnrealEd/Public/Toolkits/BaseToolkit.h
--- a/Engine/Source/Editor/UnrealEd/Public/Toolkits/BaseToolkit.h
+++ b/Engine/Source/Editor/UnrealEd/Public/Toolkits/BaseToolkit.h
@@ -270,4 +270,7 @@ protected:
 
 	/** returns whether or not this FModeToolkit has a ToolkitBuilder defined */
 	bool HasToolkitBuilder() const;
+
+	/** The sections of the toolkit, defined for the ToolkitBuilder (if present) */
+	TSharedPtr<FToolkitSections> ToolkitSections;
 };
diff --git a/Engine/Source/Runtime/Slate/Private/Framework/MultiBox/MultiBox.cpp b/Engine/Source/Runtime/Slate/Private/Framework/MultiBox/MultiBox.cpp
--- a/Engine/Source/Runtime/Slate/Private/Framework/MultiBox/MultiBox.cpp
+++ b/Engine/Source/Runtime/Slate/Private/Framework/MultiBox/MultiBox.cpp
@@ -817,7 +817,6 @@ void SMultiBoxWidget::AddBlockWidget(const FMultiBlock& Block, TSharedPtr<SHoriz
 			{
 				VerticalBox->AddSlot()
 				.AutoHeight()
-				.Padding(0.0f, 1.0f, 0.0f, 1.0f)
 				[
 					FinalWidgetWithHook
 				];
@@ -955,7 +954,7 @@ void SMultiBoxWidget::BuildMultiBoxWidget()
 	TSharedPtr<SHorizontalBox> ButtonRow;
 
 	TSharedPtr< STileView< TSharedPtr<SWidget> > > TileView;
-
+	const FToolBarStyle& Style = StyleSet->GetWidgetStyle<FToolBarStyle>(MultiBox->GetStyleName());
 
 	switch (MultiBox->GetType())
 	{
@@ -1005,17 +1004,17 @@ void SMultiBoxWidget::BuildMultiBoxWidget()
 			MainWidget = VerticalBox = SNew (SVerticalBox)
 
 			+SVerticalBox::Slot()
-			.Padding(FMargin(0.f, 2.f))
+			.Padding(0.f)
 			.AutoHeight()
 			[
 				SAssignNew(UniformToolbarPanel, SUniformWrapPanel)
 				.HAlign(HAlign_Fill)
-				.MinDesiredSlotWidth(150.f)
-				.MinDesiredSlotHeight(32.f)
-				.MaxDesiredSlotWidth(150.f)
-				.MaxDesiredSlotHeight(32.f)
-				.NumColumnsOverride(2)
-				.SlotPadding(FMargin(2.f, 1.f))
+				.MinDesiredSlotWidth(Style.UniformBlockWidth)
+				.MinDesiredSlotHeight(Style.UniformBlockHeight)
+				.MaxDesiredSlotWidth(Style.UniformBlockWidth)
+				.MaxDesiredSlotHeight(Style.UniformBlockHeight)
+				.NumColumnsOverride(Style.NumColumns)
+				.SlotPadding(0.f)	
 			];
 		}
 		break;
@@ -1264,7 +1263,11 @@ void SMultiBoxWidget::BuildMultiBoxWidget()
 TSharedRef<SWidget> SMultiBoxWidget::OnWrapButtonClicked()
 {
 	FMenuBuilder MenuBuilder(true, MultiBox->GetLastCommandList(), TSharedPtr<FExtender>(), false, GetStyleSet());
-	{ 
+	{
+		if (ClippedVerticalBox)
+		{
+			MenuBuilder.SetCheckBoxStyle("ClippingVerticalBox.Check");
+		}
 		const TArray< TSharedRef< const FMultiBlock > >& Blocks = MultiBox->GetBlocks();
 
 		int32 BlockIdx = ClippedHorizontalBox ? ClippedHorizontalBox->GetClippedIndex() : ClippedVerticalBox->GetClippedIndex();
diff --git a/Engine/Source/Runtime/Slate/Private/Framework/MultiBox/MultiBoxBuilder.cpp b/Engine/Source/Runtime/Slate/Private/Framework/MultiBox/MultiBoxBuilder.cpp
--- a/Engine/Source/Runtime/Slate/Private/Framework/MultiBox/MultiBoxBuilder.cpp
+++ b/Engine/Source/Runtime/Slate/Private/Framework/MultiBox/MultiBoxBuilder.cpp
@@ -28,12 +28,19 @@ FMultiBoxBuilder::FMultiBoxBuilder( const EMultiBoxType InType, FMultiBoxCustomi
 	, CommandListStack()
 	, TutorialHighlightName(InTutorialHighlightName)
 	, MenuName(InMenuName)
+    , CheckBoxStyle(NAME_None)
 	, bExtendersEnabled(true)
 {
 	CommandListStack.Push( InCommandList );
 	ExtenderStack.Push(InExtender);
 }
 
+void FMultiBoxBuilder::SetCheckBoxStyle(FName InCheckBoxStyle)
+{
+	this->CheckBoxStyle = InCheckBoxStyle;
+}
+
+
 void FMultiBoxBuilder::AddEditableText( const FText& InLabel, const FText& InToolTip, const FSlateIcon& InIcon, const TAttribute< FText >& InTextAttribute, const FOnTextCommitted& InOnTextCommitted, const FOnTextChanged& InOnTextChanged, bool bInReadOnly )
 {
 	MultiBox->AddMultiBlock( MakeShareable( new FEditableTextBlock( InLabel, InToolTip, InIcon, InTextAttribute, bInReadOnly, InOnTextCommitted, InOnTextChanged ) ) );
@@ -149,6 +156,7 @@ void FBaseMenuBuilder::AddMenuEntry( const TSharedPtr< const FUICommandInfo > In
 	check( InCommand.IsValid() );
 	TSharedRef< FMenuEntryBlock > NewMenuEntryBlock = MakeShared<FMenuEntryBlock>( InExtensionHook, InCommand, CommandListStack.Last(), InLabelOverride, InToolTipOverride, InIconOverride, bCloseSelfOnly );
 	NewMenuEntryBlock->SetTutorialHighlightName(GenerateTutorialIdentfierName(TutorialHighlightName, InTutorialHighlightName, InCommand, MultiBox->GetBlocks().Num()));
+	NewMenuEntryBlock->SetCheckBoxStyle(CheckBoxStyle);
 	MultiBox->AddMultiBlock( NewMenuEntryBlock );
 
 	ApplyHook(InExtensionHook, EExtensionHook::After);
@@ -274,7 +282,8 @@ void FMenuBuilder::AddSubMenu( const TAttribute<FText>& InMenuLabel, const TAttr
 	const bool bIsSubMenu = true;
 	TSharedRef< FMenuEntryBlock > NewMenuEntryBlock = MakeShared<FMenuEntryBlock>( InExtensionHook, InMenuLabel, InToolTip, InSubMenu, ExtenderStack.Top(), bIsSubMenu, bInOpenSubMenuOnClick, CommandListStack.Last(), bCloseSelfOnly, InIcon, bInShouldCloseWindowAfterMenuSelection );
 	NewMenuEntryBlock->SetRecursivelySearchable(bRecursivelySearchable);
-
+	NewMenuEntryBlock->SetCheckBoxStyle(CheckBoxStyle);
+	
 	MultiBox->AddMultiBlock( NewMenuEntryBlock );
 }
 
diff --git a/Engine/Source/Runtime/Slate/Private/Framework/MultiBox/SMenuEntryBlock.cpp b/Engine/Source/Runtime/Slate/Private/Framework/MultiBox/SMenuEntryBlock.cpp
--- a/Engine/Source/Runtime/Slate/Private/Framework/MultiBox/SMenuEntryBlock.cpp
+++ b/Engine/Source/Runtime/Slate/Private/Framework/MultiBox/SMenuEntryBlock.cpp
@@ -449,6 +449,11 @@ private:
 	EResponseToMouseUp ResponseToMouseUp;
 };
 
+void FMenuEntryBlock::SetCheckBoxStyle(FName InCheckBoxStyle)
+{
+	this->CheckBoxStyle = InCheckBoxStyle;
+}
+
 TSharedRef< SWidget > SMenuEntryBlock::BuildMenuEntryWidget( const FMenuEntryBuildParams& InBuildParams )
 {
 	const TAttribute<FText>& Label = InBuildParams.Label;
@@ -538,7 +543,12 @@ TSharedRef< SWidget > SMenuEntryBlock::BuildMenuEntryWidget( const FMenuEntryBui
 
 	TAttribute<FSlateColor> CheckBoxForegroundColor = FSlateColor::UseStyle();
 	FName CheckBoxStyle = ISlateStyle::Join( StyleName, ".CheckBox" );
-	if (UserInterfaceType == EUserInterfaceActionType::Check)
+
+	if (MenuEntryBlock->CheckBoxStyle != NAME_None)
+	{
+		CheckBoxStyle = MenuEntryBlock->CheckBoxStyle;
+	}
+	else if (UserInterfaceType == EUserInterfaceActionType::Check)
 	{
 		CheckBoxStyle = ISlateStyle::Join( StyleName, ".Check" );
 	}
diff --git a/Engine/Source/Runtime/Slate/Private/Framework/MultiBox/SMenuEntryBlock.h b/Engine/Source/Runtime/Slate/Private/Framework/MultiBox/SMenuEntryBlock.h
--- a/Engine/Source/Runtime/Slate/Private/Framework/MultiBox/SMenuEntryBlock.h
+++ b/Engine/Source/Runtime/Slate/Private/Framework/MultiBox/SMenuEntryBlock.h
@@ -120,6 +120,9 @@ public:
 	/** FMultiBlock interface */
 	virtual void CreateMenuEntry(class FMenuBuilder& MenuBuilder) const override;
 	virtual bool HasIcon() const override;
+	
+	/** the override for the checkbox style */
+	void SetCheckBoxStyle(FName InCheckBoxStyle);
 
 	/** Returns whether this menu entry block opens a sub-menu. */
 	bool IsSubMenu() const { return bIsSubMenu; }
@@ -186,6 +189,9 @@ private:
 
 	/** Whether to invert the label text's color on hover */
 	bool bInvertLabelOnHover;
+
+	/** override of the checkbox style */
+	FName CheckBoxStyle = NAME_None;
 };
 
 
diff --git a/Engine/Source/Runtime/Slate/Private/Framework/MultiBox/SToolBarButtonBlock.cpp b/Engine/Source/Runtime/Slate/Private/Framework/MultiBox/SToolBarButtonBlock.cpp
--- a/Engine/Source/Runtime/Slate/Private/Framework/MultiBox/SToolBarButtonBlock.cpp
+++ b/Engine/Source/Runtime/Slate/Private/Framework/MultiBox/SToolBarButtonBlock.cpp
@@ -234,7 +234,7 @@ void SToolBarButtonBlock::BuildMultiBlockWidget(const ISlateStyle* StyleSet, con
 			.AddMetaData<FTagMetaData>(FTagMetaData(TutorialHighlightName))
 			+ SHorizontalBox::Slot()
 			.AutoWidth()
-		//    .Padding(ToolBarStyle.IconPadding) @TODO: ~enable with styling changes
+		    .Padding(ToolBarStyle.IconPadding)
 			.VAlign(VAlign_Center)
 			.HAlign(HAlign_Center)
 			[
@@ -265,6 +265,7 @@ void SToolBarButtonBlock::BuildMultiBlockWidget(const ISlateStyle* StyleSet, con
 				// Icon image
 				+ SVerticalBox::Slot()
 				.AutoHeight()
+				.Padding(ToolBarStyle.IconPadding)
 				.HAlign(HAlign_Center)	// Center the icon horizontally, so that large labels don't stretch out the artwork
 				[
 					IconWidget
diff --git a/Engine/Source/Runtime/Slate/Private/Widgets/Layout/SUniformWrapPanel.cpp b/Engine/Source/Runtime/Slate/Private/Widgets/Layout/SUniformWrapPanel.cpp
--- a/Engine/Source/Runtime/Slate/Private/Widgets/Layout/SUniformWrapPanel.cpp
+++ b/Engine/Source/Runtime/Slate/Private/Widgets/Layout/SUniformWrapPanel.cpp
@@ -221,7 +221,7 @@ void SUniformWrapPanel::SetMaxDesiredSlotWidth(TAttribute<float> InMaxDesiredSlo
 	MaxDesiredSlotWidth.Assign(*this, MoveTemp(InMaxDesiredSlotWidth));
 }
 
-void SUniformWrapPanel::SetNumColumnsOverride(TAttribute<uint32> InNumColumnsOverride)
+void SUniformWrapPanel::SetNumColumnsOverride(TAttribute<int32> InNumColumnsOverride)
 {
 	NumColumnsOverride.Assign(*this, MoveTemp(InNumColumnsOverride));
 }
diff --git a/Engine/Source/Runtime/Slate/Public/Framework/MultiBox/MultiBoxBuilder.h b/Engine/Source/Runtime/Slate/Public/Framework/MultiBox/MultiBoxBuilder.h
--- a/Engine/Source/Runtime/Slate/Public/Framework/MultiBox/MultiBoxBuilder.h
+++ b/Engine/Source/Runtime/Slate/Public/Framework/MultiBox/MultiBoxBuilder.h
@@ -127,6 +127,9 @@ public:
 	/** @return The style name used by the multibox widgets */
 	const FName& GetStyleName() const;
 
+	/** the override for the checkbox style */
+	void SetCheckBoxStyle(FName InCheckBoxStyle);
+
 	/** Sets the style to use on the multibox widgets */
 	void SetStyle( const ISlateStyle* InStyleSet, const FName& InStyleName );
 
@@ -164,6 +167,9 @@ protected:
 	/** Name of the menu */
 	FName MenuName;
 
+	/** the override for the checkbox style for this menu */
+	FName CheckBoxStyle;
+
 	/** If extenders are enabled */
 	bool bExtendersEnabled;
 };
diff --git a/Engine/Source/Runtime/Slate/Public/Widgets/Layout/SUniformWrapPanel.h b/Engine/Source/Runtime/Slate/Public/Widgets/Layout/SUniformWrapPanel.h
--- a/Engine/Source/Runtime/Slate/Public/Widgets/Layout/SUniformWrapPanel.h
+++ b/Engine/Source/Runtime/Slate/Public/Widgets/Layout/SUniformWrapPanel.h
@@ -68,7 +68,7 @@ public:
 		SLATE_ATTRIBUTE(FMargin, SlotPadding)
 
 		/** The number of columns for the wrapped panel */
-		SLATE_ATTRIBUTE(uint32, NumColumnsOverride)
+		SLATE_ATTRIBUTE(int32, NumColumnsOverride)
 	
 		/** The minimum desired width of the slots */
 		SLATE_ATTRIBUTE(float, MinDesiredSlotWidth)
@@ -107,7 +107,7 @@ public:
 	void SetMinDesiredSlotHeight(TAttribute<float> InMinDesiredSlotHeight);
 
 	/** See NumColumnsOverride attribute */
-	void SetNumColumnsOverride(TAttribute<uint32> InNumColumnsOverride);
+	void SetNumColumnsOverride(TAttribute<int32> InNumColumnsOverride);
 	
 	/** See MinDesiredSlotWidth attribute */
 	void SetMaxDesiredSlotWidth(TAttribute<float> InMaxDesiredSlotWidth);
@@ -153,7 +153,7 @@ private:
 	mutable int32 NumRows;
 	mutable int32 NumVisibleChildren;
 
-	TSlateAttribute<uint32> NumColumnsOverride;
+	TSlateAttribute<int32> NumColumnsOverride;
 	
 	TSlateAttribute<float> MinDesiredSlotWidth;
 	TSlateAttribute<float> MinDesiredSlotHeight;
diff --git a/Engine/Source/Runtime/SlateCore/Private/Styling/StarshipCoreStyle.cpp b/Engine/Source/Runtime/SlateCore/Private/Styling/StarshipCoreStyle.cpp
--- a/Engine/Source/Runtime/SlateCore/Private/Styling/StarshipCoreStyle.cpp
+++ b/Engine/Source/Runtime/SlateCore/Private/Styling/StarshipCoreStyle.cpp
@@ -2348,6 +2348,7 @@ void FStarshipCoreStyle::SetupMultiboxStyles(TSharedRef<FStyle>& Style)
 		/* ...and add the new style */
 		Style->Set("Menu.CheckBox", MenuCheckBox);
 		Style->Set("Menu.Check", MenuCheckIndicator);
+		Style->Set("ClippingVerticalBox.Check", MenuCheckIndicator);
 
 		/* This radio button is actually just a check box with different images */
 		/* Set images for various Menu radio button (SCheckBox) states... */
diff --git a/Engine/Source/Runtime/SlateCore/Private/Styling/ToolBarStyle.cpp b/Engine/Source/Runtime/SlateCore/Private/Styling/ToolBarStyle.cpp
--- a/Engine/Source/Runtime/SlateCore/Private/Styling/ToolBarStyle.cpp
+++ b/Engine/Source/Runtime/SlateCore/Private/Styling/ToolBarStyle.cpp
@@ -17,6 +17,7 @@ FToolBarStyle::FToolBarStyle()
 	, SettingsComboButton()
 	, ButtonStyle()
 	, LabelPadding(0)
+	, IconPadding(0)
 	, SeparatorPadding(0)
 	, ComboButtonPadding(0)
 	, ButtonPadding(0)
diff --git a/Engine/Source/Runtime/SlateCore/Public/Styling/ToolBarStyle.h b/Engine/Source/Runtime/SlateCore/Public/Styling/ToolBarStyle.h
--- a/Engine/Source/Runtime/SlateCore/Public/Styling/ToolBarStyle.h
+++ b/Engine/Source/Runtime/SlateCore/Public/Styling/ToolBarStyle.h
@@ -77,6 +77,22 @@ struct SLATECORE_API FToolBarStyle : public FSlateWidgetStyle
 	FMargin LabelPadding;
 	FToolBarStyle& SetLabelPadding(const FMargin& InLabelPadding) { LabelPadding = InLabelPadding; return *this; }
 
+	UPROPERTY(EditAnywhere, BlueprintReadWrite, Category = Appearance)
+	float UniformBlockWidth;
+	FToolBarStyle& SetUniformBlockWidth(const float InUniformBlockWidth) { UniformBlockWidth = InUniformBlockWidth; return *this; }
+
+	UPROPERTY(EditAnywhere, BlueprintReadWrite, Category = Appearance)
+	float UniformBlockHeight;
+	FToolBarStyle& SetUniformBlockHeight(const float InUniformBlockHeight) { UniformBlockHeight = InUniformBlockHeight; return *this; }
+	
+	UPROPERTY(EditAnywhere, BlueprintReadWrite, Category = Appearance)
+	int32 NumColumns;
+	FToolBarStyle& SetNumColumns(const int32 InNumColumns) { NumColumns = InNumColumns; return *this; }
+	
+	UPROPERTY(EditAnywhere, BlueprintReadWrite, Category = Appearance)
+	FMargin IconPadding;
+	FToolBarStyle& SetIconPadding(const FMargin& InIconPadding) { IconPadding = InIconPadding; return *this; }
+
 	UPROPERTY(EditAnywhere, BlueprintReadWrite, Category = Appearance)
 	FMargin SeparatorPadding;
 	FToolBarStyle& SetSeparatorPadding(const FMargin& InSeparatorPadding) { SeparatorPadding = InSeparatorPadding; return *this; }
