MeshModelingtoolset: Sanitize new asset name to alphanumeric in restrictive mode

diff --git a/Engine/Plugins/Editor/ModelingToolsEditorMode/Source/ModelingToolsEditorMode/Private/ModelingModeAssetUtils.cpp b/Engine/Plugins/Editor/ModelingToolsEditorMode/Source/ModelingToolsEditorMode/Private/ModelingModeAssetUtils.cpp
--- a/Engine/Plugins/Editor/ModelingToolsEditorMode/Source/ModelingToolsEditorMode/Private/ModelingModeAssetUtils.cpp
+++ b/Engine/Plugins/Editor/ModelingToolsEditorMode/Source/ModelingToolsEditorMode/Private/ModelingModeAssetUtils.cpp
@@ -65,6 +65,18 @@ FString InteractiveSelectAssetPath(const FString& DefaultAssetName, const FText&
 }
 
 
+// Replaces each non-alphanumeric character with an underscore.
+void SanitizeObjectNameToAlphaNumeric(FString& ObjectName)
+{
+	for (TCHAR& Character : ObjectName)
+	{
+		if (!FChar::IsAlnum(Character))
+		{
+			Character = TEXT('_');
+		}
+	}
+}
+
 
 }
 }
@@ -209,6 +221,12 @@ FString UE::Modeling::GetNewAssetPathName(const FString& BaseNameIn, const UWorl
 	}
 
 	FString UseBaseName = ObjectBaseName;
+
+	if (Settings->InRestrictiveMode())
+	{
+		Local::SanitizeObjectNameToAlphaNumeric(UseBaseName);
+	}
+
 	if (Settings->bAppendRandomStringToName)
 	{
 		FString GuidString = UE::Modeling::GenerateRandomShortHexString();
