MeshModelingToolset: Disable ISM related properties/tools in restrictive mode

diff --git a/Engine/Plugins/Editor/ModelingToolsEditorMode/Source/ModelingToolsEditorMode/Private/ModelingToolsEditorMode.cpp b/Engine/Plugins/Editor/ModelingToolsEditorMode/Source/ModelingToolsEditorMode/Private/ModelingToolsEditorMode.cpp
--- a/Engine/Plugins/Editor/ModelingToolsEditorMode/Source/ModelingToolsEditorMode/Private/ModelingToolsEditorMode.cpp
+++ b/Engine/Plugins/Editor/ModelingToolsEditorMode/Source/ModelingToolsEditorMode/Private/ModelingToolsEditorMode.cpp
@@ -591,7 +591,10 @@ void UModelingToolsEditorMode::Enter()
 	RegisterTool(ToolManagerCommands.BeginTransformUVIslandsTool, TEXT("BeginTransformUVIslandsTool"), NewObject<UEditUVIslandsToolBuilder>());
 	RegisterTool(ToolManagerCommands.BeginLatticeDeformerTool, TEXT("BeginLatticeDeformerTool"), NewObject<ULatticeDeformerToolBuilder>());
 	RegisterTool(ToolManagerCommands.BeginSubdividePolyTool, TEXT("BeginSubdividePolyTool"), NewObject<USubdividePolyToolBuilder>());
-	RegisterTool(ToolManagerCommands.BeginPatternTool, TEXT("BeginPatternTool"), NewObject<UPatternToolBuilder>());
+
+	UPatternToolBuilder* PatternToolBuilder = NewObject<UPatternToolBuilder>();
+	PatternToolBuilder->bEnableCreateISMCs = !ModelingModeSettings->InRestrictiveMode();
+	RegisterTool(ToolManagerCommands.BeginPatternTool, TEXT("BeginPatternTool"), PatternToolBuilder);
 
 	UCombineMeshesToolBuilder* CombineMeshesToolBuilder = NewObject<UCombineMeshesToolBuilder>();
 	RegisterTool(ToolManagerCommands.BeginCombineMeshesTool, TEXT("BeginCombineMeshesTool"), CombineMeshesToolBuilder);
@@ -608,8 +611,12 @@ void UModelingToolsEditorMode::Enter()
 	GenerateSMLODToolBuilder->bInRestrictiveMode = ModelingModeSettings->InRestrictiveMode();
 	RegisterTool(ToolManagerCommands.BeginGenerateStaticMeshLODAssetTool, TEXT("BeginGenerateStaticMeshLODAssetTool"), GenerateSMLODToolBuilder);
 
-	UISMEditorToolBuilder* ISMEditorToolBuilder = NewObject<UISMEditorToolBuilder>();
-	RegisterTool(ToolManagerCommands.BeginISMEditorTool, TEXT("BeginISMEditorTool"), ISMEditorToolBuilder);
+	// ISMEd is disabled in Restrictive Mode.
+	if (ToolManagerCommands.BeginISMEditorTool)
+	{
+		UISMEditorToolBuilder* ISMEditorToolBuilder = NewObject<UISMEditorToolBuilder>();
+		RegisterTool(ToolManagerCommands.BeginISMEditorTool, TEXT("BeginISMEditorTool"), ISMEditorToolBuilder);
+	}
 
 
 	// edit tools
diff --git a/Engine/Plugins/Editor/ModelingToolsEditorMode/Source/ModelingToolsEditorMode/Private/ModelingToolsEditorModeToolkit.cpp b/Engine/Plugins/Editor/ModelingToolsEditorMode/Source/ModelingToolsEditorMode/Private/ModelingToolsEditorModeToolkit.cpp
--- a/Engine/Plugins/Editor/ModelingToolsEditorMode/Source/ModelingToolsEditorMode/Private/ModelingToolsEditorModeToolkit.cpp
+++ b/Engine/Plugins/Editor/ModelingToolsEditorMode/Source/ModelingToolsEditorMode/Private/ModelingToolsEditorModeToolkit.cpp
@@ -371,9 +371,12 @@ void FModelingToolsEditorModeToolkit::RegisterPalettes()
 
 	TArray<TSharedPtr<FUICommandInfo>> LODPaletteItems({
 		Commands.BeginLODManagerTool,
-		Commands.BeginGenerateStaticMeshLODAssetTool,
-		Commands.BeginISMEditorTool
+		Commands.BeginGenerateStaticMeshLODAssetTool
 	});
+	if (Commands.BeginISMEditorTool)
+	{
+		LODPaletteItems.Add(Commands.BeginISMEditorTool);
+	}
 	ToolkitBuilder->AddPalette(
 		MakeShareable( new FToolPalette( Commands.LoadLodsTools.ToSharedRef(), LODPaletteItems)) );
 
@@ -1488,7 +1491,10 @@ void FModelingToolsEditorModeToolkit::BuildToolPalette(FName PaletteIndex, class
 	{
 		ToolbarBuilder.AddToolBarButton(Commands.BeginLODManagerTool);
 		ToolbarBuilder.AddToolBarButton(Commands.BeginGenerateStaticMeshLODAssetTool);
-		ToolbarBuilder.AddToolBarButton(Commands.BeginISMEditorTool);
+		if (Commands.BeginISMEditorTool)
+		{
+			ToolbarBuilder.AddToolBarButton(Commands.BeginISMEditorTool);
+		}
 	}
 	else if (PaletteIndex == VoxToolsTabName)
 	{
diff --git a/Engine/Plugins/Editor/ModelingToolsEditorMode/Source/ModelingToolsEditorMode/Private/ModelingToolsManagerActions.cpp b/Engine/Plugins/Editor/ModelingToolsEditorMode/Source/ModelingToolsEditorMode/Private/ModelingToolsManagerActions.cpp
--- a/Engine/Plugins/Editor/ModelingToolsEditorMode/Source/ModelingToolsEditorMode/Private/ModelingToolsManagerActions.cpp
+++ b/Engine/Plugins/Editor/ModelingToolsEditorMode/Source/ModelingToolsEditorMode/Private/ModelingToolsManagerActions.cpp
@@ -309,13 +309,19 @@ void FModelingToolsManagerCommands::RegisterCommands()
 	{
 		REGISTER_MODELING_TOOL_COMMAND(BeginLODManagerTool, "LODMgr", "Inspect the LODs of a Static Mesh Asset");
 		REGISTER_MODELING_TOOL_COMMAND(BeginGenerateStaticMeshLODAssetTool, "AutoLOD", "Automatically generate a simplified LOD with baked Textures/Materials for a Mesh");
-		REGISTER_MODELING_TOOL_COMMAND(BeginISMEditorTool, "ISMEd", "Edit the Instances of Instanced Static Mesh Components");
+		if (!Settings->InRestrictiveMode())
+		{
+			REGISTER_MODELING_TOOL_COMMAND(BeginISMEditorTool, "ISMEd", "Edit the Instances of Instanced Static Mesh Components");
+		}
 	}
 	else
 	{
 		REGISTER_MODELING_TOOL_COMMAND(BeginLODManagerTool, "LOD Manager", "Inspect the LODs of a Static Mesh Asset");
 		REGISTER_MODELING_TOOL_COMMAND(BeginGenerateStaticMeshLODAssetTool, "AutoLOD", "Automatically generate a simplified LOD with baked Textures/Materials for a Mesh");
-		REGISTER_MODELING_TOOL_COMMAND(BeginISMEditorTool, "ISM Editor", "Edit the Instances of Instanced Static Mesh Components");
+		if (!Settings->InRestrictiveMode())
+		{
+			REGISTER_MODELING_TOOL_COMMAND(BeginISMEditorTool, "ISM Editor", "Edit the Instances of Instanced Static Mesh Components");
+		}
 	}
 
 	REGISTER_MODELING_TOOL_COMMAND(BeginAddPatchTool, "Patch", "Add Patch");
diff --git a/Engine/Plugins/Experimental/MeshModelingToolsetExp/Source/MeshModelingToolsExp/Private/PatternTool.cpp b/Engine/Plugins/Experimental/MeshModelingToolsetExp/Source/MeshModelingToolsExp/Private/PatternTool.cpp
--- a/Engine/Plugins/Experimental/MeshModelingToolsetExp/Source/MeshModelingToolsExp/Private/PatternTool.cpp
+++ b/Engine/Plugins/Experimental/MeshModelingToolsetExp/Source/MeshModelingToolsExp/Private/PatternTool.cpp
@@ -69,7 +69,10 @@ bool UPatternToolBuilder::CanBuildTool(const FToolBuilderState& SceneState) cons
 
 UMultiSelectionMeshEditingTool* UPatternToolBuilder::CreateNewTool(const FToolBuilderState& SceneState) const
 {
-	return NewObject<UPatternTool>(SceneState.ToolManager);
+	UPatternTool* NewTool = NewObject<UPatternTool>(SceneState.ToolManager);
+	NewTool->SetEnableCreateISMCs(bEnableCreateISMCs);
+	return NewTool;
+	
 }
 
 void UPatternToolBuilder::InitializeNewTool(UMultiSelectionMeshEditingTool* NewTool, const FToolBuilderState& SceneState) const
@@ -704,6 +707,8 @@ void UPatternTool::Setup()
 	OutputSettings = NewObject<UPatternTool_OutputSettings>();
 	AddToolPropertySource(OutputSettings);
 	OutputSettings->RestoreProperties(this);
+	OutputSettings->bCreateISMCs &= bEnableCreateISMCs;
+	OutputSettings->bEnableCreateISMCs = bEnableCreateISMCs;
 
 	BoundingBoxVisualizer.LineThickness = 2.0f;
 	
@@ -1060,6 +1065,15 @@ void UPatternTool::OnParametersUpdated()
 	MarkPatternDirty();
 }
 
+void UPatternTool::SetEnableCreateISMCs(bool bEnable)
+{
+	bEnableCreateISMCs = bEnable;
+	if (OutputSettings)
+	{
+		OutputSettings->bCreateISMCs &= bEnableCreateISMCs;
+		OutputSettings->bEnableCreateISMCs = bEnableCreateISMCs;
+	}
+}
 
 void UPatternTool::OnTransformGizmoUpdated(UTransformProxy* Proxy, FTransform Transform)
 {
diff --git a/Engine/Plugins/Experimental/MeshModelingToolsetExp/Source/MeshModelingToolsExp/Public/PatternTool.h b/Engine/Plugins/Experimental/MeshModelingToolsetExp/Source/MeshModelingToolsExp/Public/PatternTool.h
--- a/Engine/Plugins/Experimental/MeshModelingToolsetExp/Source/MeshModelingToolsExp/Public/PatternTool.h
+++ b/Engine/Plugins/Experimental/MeshModelingToolsetExp/Source/MeshModelingToolsExp/Public/PatternTool.h
@@ -42,6 +42,8 @@ public:
 	virtual UMultiSelectionMeshEditingTool* CreateNewTool(const FToolBuilderState& SceneState) const override;
 
 	virtual void InitializeNewTool(UMultiSelectionMeshEditingTool* NewTool, const FToolBuilderState& SceneState) const override;
+
+	bool bEnableCreateISMCs = true;
 	
 protected:
 	virtual const FToolTargetTypeRequirements& GetTargetRequirements() const override;
@@ -403,12 +405,16 @@ public:
 	bool bConvertToDynamic = false;
 
 	/** Create InstancedStaticMeshComponents instead multiple StaticMeshComponents, for StaticMesh pattern elements */
-	UPROPERTY(EditAnywhere, Category = Output, meta = (EditCondition = "bHaveStaticMeshes == true && bSeparateActors == false && bConvertToDynamic == false", HideEditConditionToggle))
+	UPROPERTY(EditAnywhere, Category = Output, meta = (EditCondition = "bHaveStaticMeshes == true && bSeparateActors == false && bConvertToDynamic == false && bEnableCreateISMCs == true", HideEditConditionToggle))
 	bool bCreateISMCs = false;
 
 	/** internal, used to control state of Instance settings */
 	UPROPERTY(meta = (TransientToolProperty))
 	bool bHaveStaticMeshes = false;
+
+	// internal, used to disable the creation of ISMCs
+	UPROPERTY(meta = (TransientToolProperty))
+	bool bEnableCreateISMCs = true;
 };
 
 
@@ -438,6 +444,8 @@ public:
 	virtual bool HasAccept() const override { return true; }
 	virtual bool CanAccept() const override { return true; }
 
+	virtual void SetEnableCreateISMCs(bool bEnable);
+
 public:
 	UPROPERTY()
 	TObjectPtr<UPatternToolSettings> Settings;
@@ -567,7 +575,8 @@ protected:
 	
 	bool bHaveNonUniformScaleElements = false;
 
-	
+	bool bEnableCreateISMCs = true;
+
 	struct FComponentSet
 	{
 		TArray<UPrimitiveComponent*> Components;
