ModelingToolsEditorMode: Fix static analysis warning for unchecked ModelingModeSettings pointer

diff --git a/Engine/Plugins/Editor/ModelingToolsEditorMode/Source/ModelingToolsEditorMode/Private/ModelingToolsEditorMode.cpp b/Engine/Plugins/Editor/ModelingToolsEditorMode/Source/ModelingToolsEditorMode/Private/ModelingToolsEditorMode.cpp
--- a/Engine/Plugins/Editor/ModelingToolsEditorMode/Source/ModelingToolsEditorMode/Private/ModelingToolsEditorMode.cpp
+++ b/Engine/Plugins/Editor/ModelingToolsEditorMode/Source/ModelingToolsEditorMode/Private/ModelingToolsEditorMode.cpp
@@ -359,6 +359,9 @@ void UModelingToolsEditorMode::Enter()
 	const UModelingToolsEditorModeSettings* ModelingModeSettings = GetDefault<UModelingToolsEditorModeSettings>();
 	const UModelingToolsModeCustomizationSettings* ModelingEditorSettings = GetDefault<UModelingToolsModeCustomizationSettings>();
 
+	check(ModelingModeSettings);
+	check(ModelingEditorSettings);
+
 	// Register builders for tool targets that the mode uses.
 	GetInteractiveToolsContext()->TargetManager->AddTargetFactory(NewObject<UStaticMeshComponentToolTargetFactory>(GetToolManager()));
 	GetInteractiveToolsContext()->TargetManager->AddTargetFactory(NewObject<UVolumeComponentToolTargetFactory>(GetToolManager()));
