ModelingToolsEditorMode: Substitute settings pointer checks calls with conditionals

diff --git a/Engine/Plugins/Editor/ModelingToolsEditorMode/Source/ModelingToolsEditorMode/Private/ModelingToolsEditorMode.cpp b/Engine/Plugins/Editor/ModelingToolsEditorMode/Source/ModelingToolsEditorMode/Private/ModelingToolsEditorMode.cpp
--- a/Engine/Plugins/Editor/ModelingToolsEditorMode/Source/ModelingToolsEditorMode/Private/ModelingToolsEditorMode.cpp
+++ b/Engine/Plugins/Editor/ModelingToolsEditorMode/Source/ModelingToolsEditorMode/Private/ModelingToolsEditorMode.cpp
@@ -359,9 +359,6 @@ void UModelingToolsEditorMode::Enter()
 	const UModelingToolsEditorModeSettings* ModelingModeSettings = GetDefault<UModelingToolsEditorModeSettings>();
 	const UModelingToolsModeCustomizationSettings* ModelingEditorSettings = GetDefault<UModelingToolsModeCustomizationSettings>();
 
-	check(ModelingModeSettings);
-	check(ModelingEditorSettings);
-
 	// Register builders for tool targets that the mode uses.
 	GetInteractiveToolsContext()->TargetManager->AddTargetFactory(NewObject<UStaticMeshComponentToolTargetFactory>(GetToolManager()));
 	GetInteractiveToolsContext()->TargetManager->AddTargetFactory(NewObject<UVolumeComponentToolTargetFactory>(GetToolManager()));
@@ -596,7 +593,7 @@ void UModelingToolsEditorMode::Enter()
 	RegisterTool(ToolManagerCommands.BeginSubdividePolyTool, TEXT("BeginSubdividePolyTool"), NewObject<USubdividePolyToolBuilder>());
 
 	UPatternToolBuilder* PatternToolBuilder = NewObject<UPatternToolBuilder>();
-	PatternToolBuilder->bEnableCreateISMCs = !ModelingModeSettings->InRestrictiveMode();
+	PatternToolBuilder->bEnableCreateISMCs = !(ModelingModeSettings && ModelingModeSettings->InRestrictiveMode());
 	RegisterTool(ToolManagerCommands.BeginPatternTool, TEXT("BeginPatternTool"), PatternToolBuilder);
 
 	UCombineMeshesToolBuilder* CombineMeshesToolBuilder = NewObject<UCombineMeshesToolBuilder>();
@@ -611,7 +608,7 @@ void UModelingToolsEditorMode::Enter()
 	RegisterTool(ToolManagerCommands.BeginLODManagerTool, TEXT("BeginLODManagerTool"), LODManagerToolBuilder);
 
 	UGenerateStaticMeshLODAssetToolBuilder* GenerateSMLODToolBuilder = NewObject<UGenerateStaticMeshLODAssetToolBuilder>();
-	GenerateSMLODToolBuilder->bInRestrictiveMode = ModelingModeSettings->InRestrictiveMode();
+	GenerateSMLODToolBuilder->bInRestrictiveMode = ModelingModeSettings && ModelingModeSettings->InRestrictiveMode();
 	RegisterTool(ToolManagerCommands.BeginGenerateStaticMeshLODAssetTool, TEXT("BeginGenerateStaticMeshLODAssetTool"), GenerateSMLODToolBuilder);
 
 	// ISMEd is disabled in Restrictive Mode.
