ModelingMode: harden handling of editor close, explicitly clear selection and shutdown tools to avoid issues with late mode Exit

diff --git a/Engine/Plugins/Editor/ModelingToolsEditorMode/Source/ModelingToolsEditorMode/Private/ModelingToolsEditorMode.cpp b/Engine/Plugins/Editor/ModelingToolsEditorMode/Source/ModelingToolsEditorMode/Private/ModelingToolsEditorMode.cpp
--- a/Engine/Plugins/Editor/ModelingToolsEditorMode/Source/ModelingToolsEditorMode/Private/ModelingToolsEditorMode.cpp
+++ b/Engine/Plugins/Editor/ModelingToolsEditorMode/Source/ModelingToolsEditorMode/Private/ModelingToolsEditorMode.cpp
@@ -950,6 +950,8 @@ void UModelingToolsEditorMode::Enter()
 		ModelingToolkit->InitializeAfterModeSetup();
 	}
 
+	EditorClosedEventHandle = GEditor->OnEditorClose().AddUObject(this, &UModelingToolsEditorMode::OnEditorClosed);
+
 	// Need to know about selection changes to (eg) clear mesh selections. 
 	// Listening to USelection::SelectionChangedEvent here instead of the underlying UTypedElementSelectionSet
 	// events because they do not fire at the right times, particular wrt undo/redo. 
@@ -1140,6 +1142,28 @@ void UModelingToolsEditorMode::Exit()
 	UEdMode::Exit();
 }
 
+void UModelingToolsEditorMode::OnEditorClosed()
+{
+	// On editor close, Exit() should run to clean up, but this happens very late.
+	// Close out any active Tools or Selections to mitigate any late-destruction issues.
+
+	if (SelectionManager != nullptr)
+	{
+		SelectionManager->ClearSelection();
+		SelectionManager->ClearActiveTargets();
+	}
+
+	if (GetToolManager()->HasAnyActiveTool())
+	{
+		GetToolManager()->DeactivateTool(EToolSide::Mouse, EToolShutdownType::Cancel);
+	}
+
+	if (EditorClosedEventHandle.IsValid() && GEditor)
+	{
+		GEditor->OnEditorClose().Remove(EditorClosedEventHandle);
+	}
+}
+
 void UModelingToolsEditorMode::OnToolsContextRender(IToolsContextRenderAPI* RenderAPI)
 {
 	if (SelectionInteraction)
diff --git a/Engine/Plugins/Editor/ModelingToolsEditorMode/Source/ModelingToolsEditorMode/Public/ModelingToolsEditorMode.h b/Engine/Plugins/Editor/ModelingToolsEditorMode/Source/ModelingToolsEditorMode/Public/ModelingToolsEditorMode.h
--- a/Engine/Plugins/Editor/ModelingToolsEditorMode/Source/ModelingToolsEditorMode/Public/ModelingToolsEditorMode.h
+++ b/Engine/Plugins/Editor/ModelingToolsEditorMode/Source/ModelingToolsEditorMode/Public/ModelingToolsEditorMode.h
@@ -108,6 +108,9 @@ protected:
 	FDelegateHandle TextureCreatedEventHandle;
 	FDelegateHandle SelectionModifiedEventHandle;
 
+	FDelegateHandle EditorClosedEventHandle;
+	void OnEditorClosed();
+
 	TUniquePtr<FStylusStateTracker> StylusStateTracker;
 
 	TSharedPtr<FLevelObjectsObserver> LevelObjectsObserver;
diff --git a/Engine/Plugins/Runtime/MeshModelingToolset/Source/ModelingComponentsEditorOnly/Private/Scene/LevelObjectsObserver.cpp b/Engine/Plugins/Runtime/MeshModelingToolset/Source/ModelingComponentsEditorOnly/Private/Scene/LevelObjectsObserver.cpp
--- a/Engine/Plugins/Runtime/MeshModelingToolset/Source/ModelingComponentsEditorOnly/Private/Scene/LevelObjectsObserver.cpp
+++ b/Engine/Plugins/Runtime/MeshModelingToolset/Source/ModelingComponentsEditorOnly/Private/Scene/LevelObjectsObserver.cpp
@@ -11,10 +11,14 @@
 
 
 FLevelObjectsObserver::~FLevelObjectsObserver()
-{
-	ensure(!OnActorDeletedHandle.IsValid());
-	ensure(!OnActorAddedHandle.IsValid());
-	ensure(!OnActorListChangedHandle.IsValid());
+{	// These ensures are meant to catch missing Shutdown() calls. However if the Editor is closed, Shutdown()
+	// (which is generally called by UEdMode exit) will not get chance to happen before objects are GC'd.
+	if (GEditor && GEngine)
+	{
+		ensure(!OnActorDeletedHandle.IsValid());
+		ensure(!OnActorAddedHandle.IsValid());
+		ensure(!OnActorListChangedHandle.IsValid());
+	}
 }
 
 
