ModelingMode: harden the hardening of ModelingMode cleanup on editor shutdown

diff --git a/Engine/Plugins/Editor/ModelingToolsEditorMode/Source/ModelingToolsEditorMode/Private/ModelingToolsEditorMode.cpp b/Engine/Plugins/Editor/ModelingToolsEditorMode/Source/ModelingToolsEditorMode/Private/ModelingToolsEditorMode.cpp
--- a/Engine/Plugins/Editor/ModelingToolsEditorMode/Source/ModelingToolsEditorMode/Private/ModelingToolsEditorMode.cpp
+++ b/Engine/Plugins/Editor/ModelingToolsEditorMode/Source/ModelingToolsEditorMode/Private/ModelingToolsEditorMode.cpp
@@ -1153,7 +1153,10 @@ void UModelingToolsEditorMode::OnEditorClosed()
 		SelectionManager->ClearActiveTargets();
 	}
 
-	if (GetToolManager()->HasAnyActiveTool())
+	if (GetModeManager() != nullptr 
+		&& GetInteractiveToolsContext() != nullptr 
+		&& GetToolManager() != nullptr 
+		&& GetToolManager()->HasAnyActiveTool())
 	{
 		GetToolManager()->DeactivateTool(EToolSide::Mouse, EToolShutdownType::Cancel);
 	}
