Fixed short meta class in PresetSettingsProperties making AttemptToFindShortTypeNamesInMetaData  automated test to fail

LogClass: Property ArrayProperty UPresetSettingsProperties::ActivePresetCollectionsPaths defines MetaData key "MetaClass" which contains short type name "InteractiveToolsPresetCollectionAsset". Suggested pathname: "/Script/PresetAsset.InteractiveToolsPresetCollectionAsset". Module:ModelingToolsEditorMode File:Public/ModelingToolsEditorModeToolkit.h

diff --git a/Engine/Plugins/Editor/ModelingToolsEditorMode/Source/ModelingToolsEditorMode/Public/ModelingToolsEditorModeToolkit.h b/Engine/Plugins/Editor/ModelingToolsEditorMode/Source/ModelingToolsEditorMode/Public/ModelingToolsEditorModeToolkit.h
--- a/Engine/Plugins/Editor/ModelingToolsEditorMode/Source/ModelingToolsEditorMode/Public/ModelingToolsEditorModeToolkit.h
+++ b/Engine/Plugins/Editor/ModelingToolsEditorMode/Source/ModelingToolsEditorMode/Public/ModelingToolsEditorModeToolkit.h
@@ -27,7 +27,7 @@ class UPresetSettingsProperties : public UObject
 	GENERATED_BODY()
 public:
 
-	UPROPERTY(EditAnywhere, Category = "Presets", meta=(MetaClass = "InteractiveToolsPresetCollectionAsset"))
+	UPROPERTY(EditAnywhere, Category = "Presets", meta=(MetaClass = "/Script/PresetAsset.InteractiveToolsPresetCollectionAsset"))
 	TArray< FSoftObjectPath > ActivePresetCollectionsPaths;
 };
 
