ModelingMode: bring back in-viewport selection toolbar. Add concept of Locking to IGeometrySelector, have StaticMeshSelector and VolumeSelector default to locked. Add util functions to check/manipulate locking state to UGeometrySelectionManager. Add locked-state indicator/button to selection toolbar.
#rnx

diff --git a/Engine/Plugins/Editor/ModelingToolsEditorMode/Source/ModelingToolsEditorMode/Private/ModelingModeToolkit_Toolbars.cpp b/Engine/Plugins/Editor/ModelingToolsEditorMode/Source/ModelingToolsEditorMode/Private/ModelingModeToolkit_Toolbars.cpp
--- a/Engine/Plugins/Editor/ModelingToolsEditorMode/Source/ModelingToolsEditorMode/Private/ModelingModeToolkit_Toolbars.cpp
+++ b/Engine/Plugins/Editor/ModelingToolsEditorMode/Source/ModelingToolsEditorMode/Private/ModelingModeToolkit_Toolbars.cpp
@@ -9,6 +9,7 @@
 #include "ModelingToolsEditorModeStyle.h"
 
 #include "ModelingSelectionInteraction.h"
+#include "Selection/GeometrySelectionManager.h"
 
 #include "ModelingComponentsSettings.h"
 #include "PropertySets/CreateMeshObjectTypeProperties.h"
@@ -142,7 +143,7 @@ void MakeSubMenu_ModeToggles(FMenuBuilder& MenuBuilder)
 		}));
 
 	MenuBuilder.AddMenuEntry(
-		LOCTEXT("ToggleInstancesSelection", "Select Instances"), 
+		LOCTEXT("ToggleInstancesSelection", "Enable Instance Selection"), 
 		LOCTEXT("ToggleInstancesSelection_Tooltip", "Enable/Disable support for direct selection of InstancedStaticMeshComponent Instances (via TypedElements.EnableViewportInstanceSelection cvar)"),
 		FSlateIcon(), ToggleInstancesEverywhereAction, NAME_None, EUserInterfaceActionType::ToggleButton);
 
@@ -368,6 +369,26 @@ TSharedRef<SWidget> MakeMenu_SelectionConfigSettings(FModelingToolsEditorModeToo
 }
 
 
+TSharedRef<SWidget> MakeMenu_SelectionEdits(FModelingToolsEditorModeToolkit* Toolkit)
+{
+	FMenuBuilder MenuBuilder(true, Toolkit->GetToolkitCommands());
+
+	MenuBuilder.BeginSection("Section_SelectionEdits", LOCTEXT("Section_SelectionEdits", "Selection Edits"));
+
+	const FModelingToolsManagerCommands& Commands = FModelingToolsManagerCommands::Get();
+
+	MenuBuilder.AddMenuEntry(Commands.BeginSelectionAction_SelectAll);
+	MenuBuilder.AddMenuEntry(Commands.BeginSelectionAction_ExpandToConnected);
+	MenuBuilder.AddMenuEntry(Commands.BeginSelectionAction_Invert);
+	MenuBuilder.AddMenuEntry(Commands.BeginSelectionAction_InvertConnected);
+	MenuBuilder.AddMenuEntry(Commands.BeginSelectionAction_Expand);
+	MenuBuilder.AddMenuEntry(Commands.BeginSelectionAction_Contract);
+
+	TSharedRef<SWidget> MenuWidget = MenuBuilder.MakeWidget();
+	return MenuWidget;
+}
+
+
 } // end namespace UELocal
 
 
@@ -377,10 +398,6 @@ TSharedRef<SWidget> FModelingToolsEditorModeToolkit::MakeMenu_ModelingModeConfig
 
 	FMenuBuilder MenuBuilder(true, TSharedPtr<FUICommandList>());
 
-	MenuBuilder.BeginSection("Section_Selection", LOCTEXT("Section_Selection", "Instances"));
-	MakeSubMenu_ModeToggles(MenuBuilder);
-	MenuBuilder.EndSection();
-
 	MenuBuilder.BeginSection("Section_Gizmo", LOCTEXT("Section_Gizmo", "Gizmo"));
 	MakeSubMenu_GizmoVisibilityMode(this, MenuBuilder);
 	MenuBuilder.EndSection();
@@ -393,6 +410,10 @@ TSharedRef<SWidget> FModelingToolsEditorModeToolkit::MakeMenu_ModelingModeConfig
 		}));
 	MenuBuilder.EndSection();
 
+	MenuBuilder.BeginSection("Section_Selection", LOCTEXT("Section_Selection", "Instances"));
+	MakeSubMenu_ModeToggles(MenuBuilder);
+	MenuBuilder.EndSection();
+
 	// only show settings UI quick access in non-Restrictive mode
 	const UModelingToolsEditorModeSettings* Settings = GetDefault<UModelingToolsEditorModeSettings>();
 	if (!Settings->InRestrictiveMode())
@@ -412,6 +433,8 @@ TSharedRef<SWidget> FModelingToolsEditorModeToolkit::MakeMenu_ModelingModeConfig
 
 void FModelingToolsEditorModeToolkit::ExtendSecondaryModeToolbar(UToolMenu *InModeToolbarMenu)
 {
+	return;		// disable for now
+
 	/*
 	 * Sanity check because the toolbar extension should happen after the ModeUILayer setup which is currently manual
 	 * for standalone asset editors
@@ -519,4 +542,146 @@ void FModelingToolsEditorModeToolkit::ExtendSecondaryModeToolbar(UToolMenu *InMo
 
 }
 
+
+
+
+void FModelingToolsEditorModeToolkit::MakeSelectionPaletteOverlayWidget()
+{
+	FVerticalToolBarBuilder ToolbarBuilder(
+		GetToolkitCommands(),
+		FMultiBoxCustomization::None,
+		TSharedPtr<FExtender>(), /*InForceSmallIcons=*/ true);
+	ToolbarBuilder.SetLabelVisibility(EVisibility::Collapsed);
+	ToolbarBuilder.SetStyle(FModelingToolsEditorModeStyle::Get().Get(), "SelectionToolBar");
+
+	const FModelingToolsManagerCommands& Commands = FModelingToolsManagerCommands::Get();
+
+	ToolbarBuilder.AddToolBarButton(Commands.MeshSelectionModeAction_NoSelection);
+	ToolbarBuilder.AddWidget(SNew(SSpacer).Size(FVector2D(1, 1)));
+	ToolbarBuilder.AddToolBarButton(Commands.MeshSelectionModeAction_MeshTriangles);
+	ToolbarBuilder.AddToolBarButton(Commands.MeshSelectionModeAction_MeshEdges);
+	ToolbarBuilder.AddToolBarButton(Commands.MeshSelectionModeAction_MeshVertices);
+	ToolbarBuilder.AddWidget(SNew(SSpacer).Size(FVector2D(1, 1)));
+	ToolbarBuilder.AddToolBarButton(Commands.MeshSelectionModeAction_GroupFaces);
+	ToolbarBuilder.AddToolBarButton(Commands.MeshSelectionModeAction_GroupEdges);
+	ToolbarBuilder.AddToolBarButton(Commands.MeshSelectionModeAction_GroupCorners);
+	ToolbarBuilder.AddWidget(SNew(SSpacer).Size(FVector2D(1, 8)));
+
+
+	//
+	// locking toggle button is implemented using 3 separate buttons because we cannot 
+	// control enabled/disabled or background color for a toolbar button
+	//
+
+	FUIAction UnlockTargetAction(
+		FExecuteAction::CreateLambda([this] { GetMeshSelectionManager()->SetCurrentTargetsLockState(false); }),
+		FCanExecuteAction::CreateLambda([this] { return GetMeshSelectionManager()->GetAnyCurrentTargetsLocked(); }),
+		FIsActionChecked::CreateLambda([] { return false; }),
+		FIsActionButtonVisible::CreateLambda([this] {
+			return GetMeshSelectionManager()->GetMeshTopologyMode() != UGeometrySelectionManager::EMeshTopologyMode::None
+				&& GetMeshSelectionManager()->GetAnyCurrentTargetsLockable() 
+				&& GetMeshSelectionManager()->GetAnyCurrentTargetsLocked();
+		})
+	);
+	FUIAction LockTargetAction(
+		FExecuteAction::CreateLambda([this] { GetMeshSelectionManager()->SetCurrentTargetsLockState(true); }),
+		FCanExecuteAction::CreateLambda([this] { return !GetMeshSelectionManager()->GetAnyCurrentTargetsLocked(); }),
+		FIsActionChecked::CreateLambda([]() { return false; }),
+		FIsActionButtonVisible::CreateLambda([this] {
+			return GetMeshSelectionManager()->GetMeshTopologyMode() != UGeometrySelectionManager::EMeshTopologyMode::None
+				&& GetMeshSelectionManager()->GetAnyCurrentTargetsLockable() 
+				&& (GetMeshSelectionManager()->GetAnyCurrentTargetsLocked() == false);
+		})
+	);
+	FUIAction DisabledLockTargetAction(
+		FExecuteAction::CreateLambda([]{ }),
+		FCanExecuteAction::CreateLambda([] { return false; }),
+		FIsActionChecked::CreateLambda([]() { return false; }),
+		FIsActionButtonVisible::CreateLambda([this] { 
+			return GetMeshSelectionManager()->GetMeshTopologyMode() == UGeometrySelectionManager::EMeshTopologyMode::None
+				|| GetMeshSelectionManager()->GetAnyCurrentTargetsLockable() == false; 
+		})
+	);
+
+	ToolbarBuilder.BeginStyleOverride(FName("SelectionToolBar.RedButton"));
+	ToolbarBuilder.AddToolBarButton(UnlockTargetAction, NAME_None,
+		LOCTEXT("Selection_UnlockTarget", "Unlock"),
+		LOCTEXT("Selection_UnlockTarget_Tooltip", "Click to Unlock the Selected Object and allow Mesh Selections"),
+		FSlateIcon(FModelingToolsEditorModeStyle::GetStyleSetName(), "SelectionToolBarIcons.LockedTarget"));
+	ToolbarBuilder.EndStyleOverride();
+	ToolbarBuilder.BeginStyleOverride(FName("SelectionToolBar.GreenButton"));
+	ToolbarBuilder.AddToolBarButton(LockTargetAction, NAME_None,
+		LOCTEXT("Selection_LockTarget", "Lock"),
+		LOCTEXT("Selection_LockTarget_Tooltip", "Click to Lock the Selected Object"),
+		FSlateIcon(FModelingToolsEditorModeStyle::GetStyleSetName(), "SelectionToolBarIcons.UnlockedTarget"));
+	ToolbarBuilder.EndStyleOverride();
+	ToolbarBuilder.AddToolBarButton(DisabledLockTargetAction, NAME_None,
+		LOCTEXT("Selection_LockTargetDisabled", "(No Locking)"),
+		LOCTEXT("Selection_LockTargetDisabled_Tooltip", "No Active Selection Targets are Lockable"),
+		FSlateIcon(FModelingToolsEditorModeStyle::GetStyleSetName(), "SelectionToolBarIcons.UnlockedTarget"));
+
+	ToolbarBuilder.AddWidget(SNew(SSpacer).Size(FVector2D(1, 8)));
+
+	ToolbarBuilder.AddWidget(
+		SNew(SComboButton)
+		.HasDownArrow(false)
+		.MenuPlacement(EMenuPlacement::MenuPlacement_MenuRight)
+		.ComboButtonStyle(FAppStyle::Get(), "SimpleComboButton")
+		.OnGetMenuContent(FOnGetContent::CreateLambda([this]()
+			{
+				return UELocal::MakeMenu_SelectionEdits(this);
+			}))
+		.ContentPadding(FMargin(1.0f, 1.0f))
+				.ButtonContent()
+				[
+					SNew(SImage)
+					.Image(FModelingToolsEditorModeStyle::Get()->GetBrush("ModelingModeSelection.Edits_Right"))
+				]);
+
+	ToolbarBuilder.AddWidget(
+		SNew(SComboButton)
+		.HasDownArrow(false)
+		.MenuPlacement(EMenuPlacement::MenuPlacement_MenuRight)
+		.ComboButtonStyle(FAppStyle::Get(), "SimpleComboButton")
+		.OnGetMenuContent(FOnGetContent::CreateLambda([this]()
+			{
+				return UELocal::MakeMenu_SelectionConfigSettings(this);
+			}))
+		.ContentPadding(FMargin(1.0f, 1.0f))
+				.ButtonContent()
+				[
+					SNew(SImage)
+					.Image(FModelingToolsEditorModeStyle::Get()->GetBrush("ModelingMode.DefaultSettings"))
+					.ColorAndOpacity(FSlateColor::UseForeground())
+				]);
+
+
+	TSharedRef<SWidget> Toolbar = ToolbarBuilder.MakeWidget();
+
+	SAssignNew(SelectionPaletteOverlayWidget, SVerticalBox)
+		+ SVerticalBox::Slot()
+		.HAlign(HAlign_Left)
+		.VAlign(VAlign_Center)
+		.Padding(FMargin(2.0f, 0.0f, 0.f, 0.f))
+		[
+			SNew(SBorder)
+			.BorderImage(FAppStyle::Get().GetBrush("EditorViewport.OverlayBrush"))
+		.Padding(FMargin(3.0f, 6.0f, 3.f, 6.f))
+		[
+			SNew(SVerticalBox)
+
+			+ SVerticalBox::Slot()
+		.AutoHeight()
+		.Padding(FMargin(0.0, 0.f, 0.f, 0.f))
+		[
+			Toolbar
+		]
+		]
+		];
+
+}
+
+
+
+
 #undef LOCTEXT_NAMESPACE
diff --git a/Engine/Plugins/Editor/ModelingToolsEditorMode/Source/ModelingToolsEditorMode/Private/ModelingToolsEditorMode.cpp b/Engine/Plugins/Editor/ModelingToolsEditorMode/Source/ModelingToolsEditorMode/Private/ModelingToolsEditorMode.cpp
--- a/Engine/Plugins/Editor/ModelingToolsEditorMode/Source/ModelingToolsEditorMode/Private/ModelingToolsEditorMode.cpp
+++ b/Engine/Plugins/Editor/ModelingToolsEditorMode/Source/ModelingToolsEditorMode/Private/ModelingToolsEditorMode.cpp
@@ -1314,7 +1314,9 @@ bool UModelingToolsEditorMode::BoxSelect(FBox& InBox, bool InSelect)
 
 bool UModelingToolsEditorMode::FrustumSelect(const FConvexVolume& InFrustum, FEditorViewportClient* InViewportClient, bool InSelect)
 {
-	if (SelectionManager && SelectionManager->HasActiveTargets())
+	if (SelectionManager 
+		&& SelectionManager->HasActiveTargets() 
+		&& SelectionManager->GetMeshTopologyMode() != UGeometrySelectionManager::EMeshTopologyMode::None)
 	{
 		UE::Geometry::FGeometrySelectionUpdateConfig UpdateConfig;
 		UpdateConfig.ChangeType = UE::Geometry::EGeometrySelectionChangeType::Replace;
diff --git a/Engine/Plugins/Editor/ModelingToolsEditorMode/Source/ModelingToolsEditorMode/Private/ModelingToolsEditorModeStyle.cpp b/Engine/Plugins/Editor/ModelingToolsEditorMode/Source/ModelingToolsEditorMode/Private/ModelingToolsEditorModeStyle.cpp
--- a/Engine/Plugins/Editor/ModelingToolsEditorMode/Source/ModelingToolsEditorMode/Private/ModelingToolsEditorModeStyle.cpp
+++ b/Engine/Plugins/Editor/ModelingToolsEditorMode/Source/ModelingToolsEditorMode/Private/ModelingToolsEditorModeStyle.cpp
@@ -10,6 +10,8 @@
 #include "Styling/SlateStyle.h"
 #include "Styling/SlateStyleMacros.h"
 #include "Styling/ToolBarStyle.h"
+#include "Styling/StarshipCoreStyle.h"
+#include "Styling/StyleColors.h"
 
 
 #define IMAGE_PLUGIN_BRUSH( RelativePath, ... ) FSlateImageBrush( FModelingToolsEditorModeStyle::InContent( RelativePath, ".png" ), __VA_ARGS__ )
@@ -17,6 +19,7 @@
 // This is to fix the issue that SlateStyleMacros like IMAGE_BRUSH look for RootToContentDir but StyleSet->RootToContentDir is how this style is set up
 #define RootToContentDir StyleSet->RootToContentDir
 
+
 FString FModelingToolsEditorModeStyle::InContent(const FString& RelativePath, const ANSICHAR* Extension)
 {
 	static FString ContentDir = IPluginManager::Get().FindPlugin(TEXT("ModelingToolsEditorMode"))->GetContentDir();
@@ -351,7 +354,6 @@ void FModelingToolsEditorModeStyle::Initialize()
 		StyleSet->Set("ModelingToolsManagerCommands.BeginSelectionAction_Expand", new IMAGE_BRUSH_SVG("Icons/ModSelectionExpand_16", Icon20x20));
 		StyleSet->Set("ModelingToolsManagerCommands.BeginSelectionAction_Contract", new IMAGE_BRUSH_SVG("Icons/ModSelectionShrink_16", Icon20x20));
 
-		
 		StyleSet->Set("ModelingToolsManagerCommands.BeginSelectionAction_Delete", new FSlateImageBrush(StyleSet->RootToCoreContentDir(TEXT("../Editor/Slate/Icons/GeneralTools/Delete_40x.png")), Icon20x20));
 
 		StyleSet->Set("ModelingToolsManagerCommands.BeginPolyModelTool_PolyEd", 				new IMAGE_PLUGIN_BRUSH("Icons/PolyEdit_40x", 		Icon20x20));
@@ -359,12 +361,42 @@ void FModelingToolsEditorModeStyle::Initialize()
 		StyleSet->Set("ModelingToolsManagerCommands.BeginPolyModelTool_TriSel", 				new IMAGE_PLUGIN_BRUSH("Icons/MeshSelect_40x",		Icon20x20));
 		StyleSet->Set("ModelingToolsManagerCommands.BeginPolyModelTool_TriSel.Small", 			new IMAGE_PLUGIN_BRUSH("Icons/MeshSelect_40x",		Icon20x20));
 
-
 		StyleSet->Set("ModelingModeSelection.More_Right",  new IMAGE_BRUSH_SVG("Icons/SelectionToolbar_More", Icon20x20));
 		StyleSet->Set("ModelingModeSelection.Edits_Right",  new IMAGE_BRUSH_SVG("Icons/SelectionToolbar_Edits", Icon20x20));
 
 
-
+		//
+		// icons and style for the mesh selection toolbar
+		//
+		StyleSet->Set("SelectionToolBarIcons.LockedTarget", new IMAGE_BRUSH_SVG("Icons/lock-red", Icon16x16));
+		StyleSet->Set("SelectionToolBarIcons.UnlockedTarget", new IMAGE_BRUSH_SVG("Icons/lock-unlocked-green", Icon16x16));
+
+		FToolBarStyle SelectionToolbarStyle = FAppStyle::Get().GetWidgetStyle<FToolBarStyle>("EditorViewportToolBar");
+		StyleSet->Set("SelectionToolBar", SelectionToolbarStyle);
+
+		// override red-button style
+		const FButtonStyle SelectionToolbarRedButton = FButtonStyle(SelectionToolbarStyle.ButtonStyle)
+			.SetNormal(FSlateRoundedBoxBrush(FStyleColors::AccentRed, 12.f, FLinearColor(0, 0, 0, .8), 1.0))
+			.SetPressed(FSlateRoundedBoxBrush(FStyleColors::AccentRed, 12.f, FLinearColor(0, 0, 0, .8), 1.0))
+			.SetHovered(FSlateRoundedBoxBrush(FStyleColors::AccentRed, 12.f, FLinearColor(0, 0, 0, .8), 1.0))
+			.SetDisabled(FSlateRoundedBoxBrush(FStyleColors::SelectInactive, 12.f, FLinearColor(0, 0, 0, .8), 1.0))
+			.SetNormalForeground(FSlateColor::UseForeground())
+			.SetPressedForeground(FSlateColor::UseForeground())
+			.SetHoveredForeground(FSlateColor::UseForeground());
+		SelectionToolbarStyle.SetButtonStyle(SelectionToolbarRedButton);
+		StyleSet->Set("SelectionToolBar.RedButton", SelectionToolbarStyle);
+
+		// override green-button style
+		const FButtonStyle SelectionToolbarGreenButton = FButtonStyle(SelectionToolbarStyle.ButtonStyle)
+			.SetNormal(FSlateRoundedBoxBrush(FStyleColors::AccentGreen, 12.f, FLinearColor(0, 0, 0, .8), 1.0))
+			.SetPressed(FSlateRoundedBoxBrush(FStyleColors::AccentGreen, 12.f, FLinearColor(0, 0, 0, .8), 1.0))
+			.SetHovered(FSlateRoundedBoxBrush(FStyleColors::AccentGreen, 12.f, FLinearColor(0, 0, 0, .8), 1.0))
+			.SetDisabled(FSlateRoundedBoxBrush(FStyleColors::SelectInactive, 12.f, FLinearColor(0, 0, 0, .8), 1.0))
+			.SetNormalForeground(FSlateColor::UseForeground())
+			.SetPressedForeground(FSlateColor::UseForeground())
+			.SetHoveredForeground(FSlateColor::UseForeground());
+		SelectionToolbarStyle.SetButtonStyle(SelectionToolbarGreenButton);
+		StyleSet->Set("SelectionToolBar.GreenButton", SelectionToolbarStyle);
 
 		//
 		// Icons for brush falloffs in sculpt/etc tools
diff --git a/Engine/Plugins/Editor/ModelingToolsEditorMode/Source/ModelingToolsEditorMode/Private/ModelingToolsEditorModeToolkit.cpp b/Engine/Plugins/Editor/ModelingToolsEditorMode/Source/ModelingToolsEditorMode/Private/ModelingToolsEditorModeToolkit.cpp
--- a/Engine/Plugins/Editor/ModelingToolsEditorMode/Source/ModelingToolsEditorMode/Private/ModelingToolsEditorModeToolkit.cpp
+++ b/Engine/Plugins/Editor/ModelingToolsEditorMode/Source/ModelingToolsEditorMode/Private/ModelingToolsEditorModeToolkit.cpp
@@ -9,6 +9,7 @@
 #include "InteractiveToolManager.h"
 #include "ModelingToolsEditorModeSettings.h"
 #include "ModelingToolsEditorModeStyle.h"
+#include "ModelingToolsEditorMode.h"
 
 #include "Modules/ModuleManager.h"
 #include "IDetailsView.h"
@@ -118,6 +119,12 @@ FModelingToolsEditorModeToolkit::~FModelingToolsEditorModeToolkit()
 {
 	if (IsHosted())
 	{
+		if (SelectionPaletteOverlayWidget.IsValid())
+		{
+			GetToolkitHost()->RemoveViewportOverlayWidget(SelectionPaletteOverlayWidget.ToSharedRef());
+			SelectionPaletteOverlayWidget.Reset();
+		}
+
 		if (GizmoNumericalUIOverlayWidget.IsValid())
 		{
 			GetToolkitHost()->RemoveViewportOverlayWidget(GizmoNumericalUIOverlayWidget.ToSharedRef());
@@ -504,6 +511,14 @@ void FModelingToolsEditorModeToolkit::Init(const TSharedPtr<IToolkitHost>& InitT
 	// it can be bound to the buttons there.
 	MakeGizmoNumericalUIOverlayWidget();
 	GetToolkitHost()->AddViewportOverlayWidget(GizmoNumericalUIOverlayWidget.ToSharedRef());
+
+	const UModelingToolsEditorModeSettings* ModelingModeSettings = GetDefault<UModelingToolsEditorModeSettings>();
+	bool bEnableSelectionUI = ModelingModeSettings && ModelingModeSettings->bEnablePersistentSelections;
+	if (bEnableSelectionUI)
+	{
+		MakeSelectionPaletteOverlayWidget();
+		GetToolkitHost()->AddViewportOverlayWidget(SelectionPaletteOverlayWidget.ToSharedRef());
+	}
 }
 
 void FModelingToolsEditorModeToolkit::AddReferencedObjects(FReferenceCollector& Collector)
@@ -1823,6 +1838,15 @@ void FModelingToolsEditorModeToolkit::BindGizmoNumericalUI()
 }
 
 
+UGeometrySelectionManager* FModelingToolsEditorModeToolkit::GetMeshSelectionManager()
+{
+	if (CachedSelectionManager == nullptr)
+	{
+		CachedSelectionManager = Cast<UModelingToolsEditorMode>(GetScriptableEditorMode())->GetSelectionManager();
+	}
+	return CachedSelectionManager;
+}
+
 void FModelingToolsEditorModeToolkit::OnToolPaletteChanged(FName PaletteName) 
 {
 }
diff --git a/Engine/Plugins/Editor/ModelingToolsEditorMode/Source/ModelingToolsEditorMode/Public/ModelingToolsEditorMode.h b/Engine/Plugins/Editor/ModelingToolsEditorMode/Source/ModelingToolsEditorMode/Public/ModelingToolsEditorMode.h
--- a/Engine/Plugins/Editor/ModelingToolsEditorMode/Source/ModelingToolsEditorMode/Public/ModelingToolsEditorMode.h
+++ b/Engine/Plugins/Editor/ModelingToolsEditorMode/Source/ModelingToolsEditorMode/Public/ModelingToolsEditorMode.h
@@ -87,10 +87,10 @@ public:
 	}
 
 	UPROPERTY()
-	bool bEnableVolumeElementSelection = false;
+	bool bEnableVolumeElementSelection = true;
 
 	UPROPERTY()
-	bool bEnableStaticMeshElementSelection = false;
+	bool bEnableStaticMeshElementSelection = true;
 
 
 protected:
diff --git a/Engine/Plugins/Editor/ModelingToolsEditorMode/Source/ModelingToolsEditorMode/Public/ModelingToolsEditorModeSettings.h b/Engine/Plugins/Editor/ModelingToolsEditorMode/Source/ModelingToolsEditorMode/Public/ModelingToolsEditorModeSettings.h
--- a/Engine/Plugins/Editor/ModelingToolsEditorMode/Source/ModelingToolsEditorMode/Public/ModelingToolsEditorModeSettings.h
+++ b/Engine/Plugins/Editor/ModelingToolsEditorMode/Source/ModelingToolsEditorMode/Public/ModelingToolsEditorModeSettings.h
@@ -358,10 +358,10 @@ public:
 	int32 LastMeshSelectionTopologyMode = 0;
 
 	UPROPERTY()
-	bool bLastMeshSelectionVolumeToggle = false;
+	bool bLastMeshSelectionVolumeToggle = true;
 
 	UPROPERTY()
-	bool bLastMeshSelectionStaticMeshToggle = false;
+	bool bLastMeshSelectionStaticMeshToggle = true;
 
 
 };
diff --git a/Engine/Plugins/Editor/ModelingToolsEditorMode/Source/ModelingToolsEditorMode/Public/ModelingToolsEditorModeToolkit.h b/Engine/Plugins/Editor/ModelingToolsEditorMode/Source/ModelingToolsEditorMode/Public/ModelingToolsEditorModeToolkit.h
--- a/Engine/Plugins/Editor/ModelingToolsEditorMode/Source/ModelingToolsEditorMode/Public/ModelingToolsEditorModeToolkit.h
+++ b/Engine/Plugins/Editor/ModelingToolsEditorMode/Source/ModelingToolsEditorMode/Public/ModelingToolsEditorModeToolkit.h
@@ -19,6 +19,7 @@ class STransformGizmoNumericalUIOverlay;
 class IDetailsView;
 class SButton;
 class STextBlock;
+class UGeometrySelectionManager;
 class UInteractiveToolsPresetCollectionAsset;
 
 UCLASS()
@@ -99,6 +100,9 @@ public:
 	// created gizmos are bound to the numerical UI.
 	void BindGizmoNumericalUI();
 
+	// get cached UGeometrySelectionManager pointer for the current Mode
+	UGeometrySelectionManager* GetMeshSelectionManager();
+
 	// This is exposed only for the convenience of being able to create the numerical UI submenu
 	// in a non-member function in ModelingModeToolkit_Toolbars.cpp
 	TSharedPtr<STransformGizmoNumericalUIOverlay> GetGizmoNumericalUIOverlayWidget() { return GizmoNumericalUIOverlayWidget; }
@@ -115,6 +119,8 @@ private:
 	FStatusBarMessageHandle ActiveToolMessageHandle;
 	const FSlateBrush* ActiveToolIcon = nullptr;
 
+	UGeometrySelectionManager* CachedSelectionManager = nullptr;
+
 	TSharedPtr<SWidget> ToolkitWidget;
 	void UpdateActiveToolProperties();
 	void InvalidateCachedDetailPanelState(UObject* ChangedObject);
@@ -125,6 +131,9 @@ private:
 	TSharedPtr<SWidget> ToolShutdownViewportOverlayWidget;
 	void MakeToolShutdownOverlayWidget();
 	
+	TSharedPtr<SWidget> SelectionPaletteOverlayWidget;
+	void MakeSelectionPaletteOverlayWidget();
+
 	void MakeGizmoNumericalUIOverlayWidget();
 	TSharedPtr<STransformGizmoNumericalUIOverlay> GizmoNumericalUIOverlayWidget;
 
diff --git a/Engine/Plugins/Runtime/MeshModelingToolset/Source/ModelingComponents/Private/Selection/DynamicMeshSelector.cpp b/Engine/Plugins/Runtime/MeshModelingToolset/Source/ModelingComponents/Private/Selection/DynamicMeshSelector.cpp
--- a/Engine/Plugins/Runtime/MeshModelingToolset/Source/ModelingComponents/Private/Selection/DynamicMeshSelector.cpp
+++ b/Engine/Plugins/Runtime/MeshModelingToolset/Source/ModelingComponents/Private/Selection/DynamicMeshSelector.cpp
@@ -116,6 +116,8 @@ void FBaseDynamicMeshSelector::InitializeSelectionFromPredicate(
 	EInitializeSelectionMode InitializeMode,
 	const FGeometrySelection* ReferenceSelection)
 {
+	if (IsLockable() && IsLocked()) return;
+
 	const FGroupTopology* UseGroupTopology = ( SelectionOut.TopologyType == EGeometryTopologyType::Polygroup ) ? 
 		GetGroupTopology() : nullptr;
 
@@ -157,6 +159,8 @@ void FBaseDynamicMeshSelector::UpdateSelectionFromSelection(
 	const FGeometrySelectionUpdateConfig& UpdateConfig,
 	FGeometrySelectionDelta* SelectionDelta )
 {
+	if (IsLockable() && IsLocked()) return;
+
 	if ( FromSelection.IsSameType(SelectionEditor.GetSelection()) )
 	{
 		UE::Geometry::UpdateSelectionWithNewElements(&SelectionEditor, UpdateConfig.ChangeType,
@@ -198,6 +202,7 @@ bool FBaseDynamicMeshSelector::RayHitTest(
 	FInputRayHit& HitResultOut )
 {
 	HitResultOut = FInputRayHit();
+	if (IsLockable() && IsLocked()) return false;
 
 	// todo: do we need to sometimes ignore surface hits here? maybe for Volumes in (eg) edge mode...
 
@@ -240,6 +245,8 @@ void FBaseDynamicMeshSelector::UpdateSelectionViaRaycast(
 	const FGeometrySelectionUpdateConfig& UpdateConfig,
 	FGeometrySelectionUpdateResult& ResultOut)
 {
+	if (IsLockable() && IsLocked()) return;
+
 	if (SelectionEditor.GetTopologyType() == EGeometryTopologyType::Triangle)
 	{
 		UpdateSelectionViaRaycast_MeshTopology(RayInfo, SelectionEditor, UpdateConfig, ResultOut);
@@ -344,6 +351,8 @@ void FBaseDynamicMeshSelector::GetSelectionPreviewForRaycast(
 	const FWorldRayQueryInfo& RayInfo,
 	FGeometrySelectionEditor& PreviewEditor)
 {
+	if (IsLockable() && IsLocked()) return;
+
 	FGeometrySelectionUpdateResult UpdateResult;
 	UpdateSelectionViaRaycast(RayInfo, PreviewEditor,
 		FGeometrySelectionUpdateConfig(),	// defaults to add
@@ -361,6 +370,8 @@ void FBaseDynamicMeshSelector::UpdateSelectionViaShape(
 	ResultOut.bSelectionModified = false;
 	ResultOut.bSelectionMissed = true;
 
+	if (IsLockable() && IsLocked()) return;
+
 	FTransformSRT3d WorldTransform = GetWorldTransformFunc();
 	TArray<uint64> InsideElements;
 
diff --git a/Engine/Plugins/Runtime/MeshModelingToolset/Source/ModelingComponents/Private/Selection/GeometrySelectionManager.cpp b/Engine/Plugins/Runtime/MeshModelingToolset/Source/ModelingComponents/Private/Selection/GeometrySelectionManager.cpp
--- a/Engine/Plugins/Runtime/MeshModelingToolset/Source/ModelingComponents/Private/Selection/GeometrySelectionManager.cpp
+++ b/Engine/Plugins/Runtime/MeshModelingToolset/Source/ModelingComponents/Private/Selection/GeometrySelectionManager.cpp
@@ -277,6 +277,31 @@ public:
 
 
 
+class FGeometrySelectionManager_TargetLockStateChange : public FToolCommandChange
+{
+public:
+	FGeometryIdentifier TargetIdentifier;
+	bool bToState;
+
+	virtual void Apply(UObject* Object) override
+	{
+		CastChecked<UGeometrySelectionManager>(Object)->SetTargetLockStateOnUndoRedo(TargetIdentifier, bToState);
+	}
+
+	virtual void Revert(UObject* Object) override
+	{
+		CastChecked<UGeometrySelectionManager>(Object)->SetTargetLockStateOnUndoRedo(TargetIdentifier, !bToState);
+	}
+
+	virtual FString ToString() const override { return TEXT("FGeometrySelectionManager_TargetLockStateChange"); }
+
+	virtual bool HasExpired(UObject* Object) const override
+	{
+		UGeometrySelectionManager* Manager = Cast<UGeometrySelectionManager>(Object);
+		return (Manager == nullptr || IsValid(Manager) == false || Manager->HasBeenShutDown());
+	}
+};
+
 
 
 bool UGeometrySelectionManager::HasActiveTargets() const
@@ -376,6 +401,80 @@ void UGeometrySelectionManager::SynchronizeActiveTargets(
 	}
 }
 
+
+bool UGeometrySelectionManager::GetAnyCurrentTargetsLockable() const
+{
+	for (TSharedPtr<FGeometrySelectionTarget> Target : ActiveTargetReferences)
+	{
+		if (Target->Selector->IsLockable())
+		{
+			return true;
+		}
+	}
+	return false;
+}
+
+bool UGeometrySelectionManager::GetAnyCurrentTargetsLocked() const
+{
+	for (TSharedPtr<FGeometrySelectionTarget> Target : ActiveTargetReferences)
+	{
+		if (Target->Selector->IsLockable() && Target->Selector->IsLocked())
+		{
+			return true;
+		}
+	}
+	return false;
+}
+
+void UGeometrySelectionManager::SetCurrentTargetsLockState(bool bLocked)
+{
+	bool bInTransaction = false;
+
+	bool bLockStateModified = false;
+	for (TSharedPtr<FGeometrySelectionTarget> Target : ActiveTargetReferences)
+	{
+		if (Target->Selector->IsLockable() && Target->Selector->IsLocked() != bLocked)
+		{
+			Target->Selector->SetLockedState(bLocked);
+			bLockStateModified = true;
+
+			if (!bInTransaction)
+			{
+				GetTransactionsAPI()->BeginUndoTransaction((bLocked) ?
+					LOCTEXT("Lock Target", "Lock Target") : LOCTEXT("Unlock Target", "Unlock Target"));
+				bInTransaction = true;
+			}
+
+			TUniquePtr<FGeometrySelectionManager_TargetLockStateChange> Change = MakeUnique<FGeometrySelectionManager_TargetLockStateChange>();
+			Change->TargetIdentifier = Target->TargetIdentifier;
+			Change->bToState = bLocked;
+			GetTransactionsAPI()->AppendChange(this, MoveTemp(Change), LOCTEXT("Lock Target", "Lock Target"));
+		}
+	}
+
+	if (bLockStateModified)
+	{
+		ClearSelection();
+	}
+
+	if (bInTransaction)
+	{
+		GetTransactionsAPI()->EndUndoTransaction();
+	}
+}
+
+void UGeometrySelectionManager::SetTargetLockStateOnUndoRedo(FGeometryIdentifier TargetIdentifier, bool bLocked)
+{
+	for (TSharedPtr<FGeometrySelectionTarget> Target : ActiveTargetReferences)
+	{
+		if (Target->TargetIdentifier == TargetIdentifier)
+		{
+			Target->Selector->SetLockedState(bLocked);
+		}
+	}
+}
+
+
 TArray<FGeometryIdentifier> UGeometrySelectionManager::GetCurrentTargetIdentifiers() const
 {
 	TArray<FGeometryIdentifier> Result;
@@ -516,7 +615,7 @@ void UGeometrySelectionManager::ClearSelection()
 		if (ClearDelta.IsEmpty() == false)
 		{
 			TUniquePtr<FGeometrySelectionDeltaChange> ClearChange = MakeUnique<FGeometrySelectionDeltaChange>();
-			ClearChange->Identifier = Target->SelectionIdentifer;
+			ClearChange->Identifier = Target->TargetIdentifier;
 			ClearChange->Delta = MoveTemp(ClearDelta);
 			GetTransactionsAPI()->AppendChange(this, MoveTemp(ClearChange), LOCTEXT("ClearSelection", "Clear Selection"));
 		}
@@ -557,7 +656,7 @@ void UGeometrySelectionManager::UpdateSelectionViaRaycast(
 	if (ResultOut.bSelectionModified)
 	{
 		TUniquePtr<FGeometrySelectionDeltaChange> DeltaChange = MakeUnique<FGeometrySelectionDeltaChange>();
-		DeltaChange->Identifier = Target->SelectionIdentifer;
+		DeltaChange->Identifier = Target->TargetIdentifier;
 		DeltaChange->Delta = ResultOut.SelectionDelta;
 
 		GetTransactionsAPI()->BeginUndoTransaction(LOCTEXT("UpdateSelectionViaRaycast", "Change Selection"));
@@ -600,7 +699,7 @@ void UGeometrySelectionManager::UpdateSelectionViaConvex(
 	if (ResultOut.bSelectionModified)
 	{
 		TUniquePtr<FGeometrySelectionDeltaChange> DeltaChange = MakeUnique<FGeometrySelectionDeltaChange>();
-		DeltaChange->Identifier = Target->SelectionIdentifer;
+		DeltaChange->Identifier = Target->TargetIdentifier;
 		DeltaChange->Delta = ResultOut.SelectionDelta;
 
 		GetTransactionsAPI()->BeginUndoTransaction(LOCTEXT("UpdateSelectionViaConvex", "Change Selection"));
@@ -698,7 +797,7 @@ void UGeometrySelectionManager::EndTrackedSelectionChange()
 			if (InitialTrackedDelta.IsEmpty() == false)
 			{
 				TUniquePtr<FGeometrySelectionDeltaChange> InitialDeltaChange = MakeUnique<FGeometrySelectionDeltaChange>();
-				InitialDeltaChange->Identifier = Target->SelectionIdentifer;
+				InitialDeltaChange->Identifier = Target->TargetIdentifier;
 				InitialDeltaChange->Delta = MoveTemp(InitialTrackedDelta);
 				GetTransactionsAPI()->AppendChange(this, MoveTemp(InitialDeltaChange), LOCTEXT("ChangeSelection", "Change Selection"));
 			}
@@ -706,7 +805,7 @@ void UGeometrySelectionManager::EndTrackedSelectionChange()
 			if (ActiveTrackedDelta.IsEmpty() == false)
 			{
 				TUniquePtr<FGeometrySelectionDeltaChange> AccumDeltaChange = MakeUnique<FGeometrySelectionDeltaChange>();
-				AccumDeltaChange->Identifier = Target->SelectionIdentifer;
+				AccumDeltaChange->Identifier = Target->TargetIdentifier;
 				AccumDeltaChange->Delta = MoveTemp(ActiveTrackedDelta);
 				GetTransactionsAPI()->AppendChange(this, MoveTemp(AccumDeltaChange), LOCTEXT("ChangeSelection", "Change Selection"));
 			}
@@ -732,7 +831,7 @@ bool UGeometrySelectionManager::SetSelectionForComponent(UPrimitiveComponent* Co
 			if ( AfterDelta.IsEmpty() == false )
 			{
 				TUniquePtr<FGeometrySelectionReplaceChange> NewSelectionChange = MakeUnique<FGeometrySelectionReplaceChange>();
-				NewSelectionChange->Identifier = Target->Selector->GetIdentifier();
+				NewSelectionChange->Identifier = Target->TargetIdentifier; //Target->Selector->GetIdentifier();
 				NewSelectionChange->After = Target->Selection;
 				NewSelectionChange->Before = InitialSelection;
 				GetTransactionsAPI()->AppendChange(this, MoveTemp(NewSelectionChange), LOCTEXT("NewSelection", "New Selection"));
@@ -971,7 +1070,7 @@ void UGeometrySelectionManager::ExecuteSelectionCommand(UGeometrySelectionEditCo
 		// so it will not be holding onto the active Selection on Redo later
 		// (if that becomes necessary, this sequence of changes will need to become more complicated....)
 		TUniquePtr<FGeometrySelectionReplaceChange> ClearChange = MakeUnique<FGeometrySelectionReplaceChange>();
-		ClearChange->Identifier = Target->Selector->GetIdentifier();
+		ClearChange->Identifier = Target->TargetIdentifier; //Target->Selector->GetIdentifier();
 		ClearChange->Before = Target->Selection;
 		ClearChange->After.InitializeTypes(ClearChange->Before);
 		GetTransactionsAPI()->AppendChange(this, MoveTemp(ClearChange), LOCTEXT("ClearSelection", "Clear Selection"));
@@ -998,7 +1097,7 @@ void UGeometrySelectionManager::ExecuteSelectionCommand(UGeometrySelectionEditCo
 				if (Target->Selection.IsEmpty() == false)
 				{
 					TUniquePtr<FGeometrySelectionReplaceChange> NewSelectionChange = MakeUnique<FGeometrySelectionReplaceChange>();
-					NewSelectionChange->Identifier = Target->Selector->GetIdentifier();
+					NewSelectionChange->Identifier = Target->TargetIdentifier; //Target->Selector->GetIdentifier();
 					NewSelectionChange->After = Target->Selection;
 					NewSelectionChange->Before.InitializeTypes(Target->Selection);
 					GetTransactionsAPI()->AppendChange(this, MoveTemp(NewSelectionChange), LOCTEXT("NewSelection", "New Selection"));
@@ -1050,7 +1149,7 @@ void UGeometrySelectionManager::ApplyChange(IGeometrySelectionChange* Change)
 
 	for (int32 k = 0; k < ActiveTargetReferences.Num(); ++k)
 	{
-		if (ActiveTargetReferences[k]->SelectionIdentifer == Identifer) 
+		if (ActiveTargetReferences[k]->TargetIdentifier == Identifer)
 		{
 			FGeometrySelectionDelta ApplyDelta;
 			Change->ApplyChange( ActiveTargetReferences[k]->SelectionEditor.Get(), ApplyDelta);
@@ -1079,7 +1178,7 @@ void UGeometrySelectionManager::RevertChange(IGeometrySelectionChange* Change)
 
 	for (int32 k = 0; k < ActiveTargetReferences.Num(); ++k)
 	{
-		if (ActiveTargetReferences[k]->SelectionIdentifer == Identifer) 
+		if (ActiveTargetReferences[k]->TargetIdentifier == Identifer) 
 		{
 			FGeometrySelectionDelta RevertDelta;
 			Change->RevertChange( ActiveTargetReferences[k]->SelectionEditor.Get(), RevertDelta);
diff --git a/Engine/Plugins/Runtime/MeshModelingToolset/Source/ModelingComponents/Public/Selection/GeometrySelectionManager.h b/Engine/Plugins/Runtime/MeshModelingToolset/Source/ModelingComponents/Public/Selection/GeometrySelectionManager.h
--- a/Engine/Plugins/Runtime/MeshModelingToolset/Source/ModelingComponents/Public/Selection/GeometrySelectionManager.h
+++ b/Engine/Plugins/Runtime/MeshModelingToolset/Source/ModelingComponents/Public/Selection/GeometrySelectionManager.h
@@ -125,6 +125,12 @@ public:
 	 */
 	bool AddActiveTarget(FGeometryIdentifier Target);
 
+
+	bool GetAnyCurrentTargetsLockable() const;
+	bool GetAnyCurrentTargetsLocked() const;
+	void SetCurrentTargetsLockState(bool bLocked);
+
+
 	/**
 	 * Update the current active target set based on DesiredActiveSet, assuming that a valid IGeometrySelectorFactory
 	 * can be found for each identifier. 
@@ -400,6 +406,9 @@ protected:
 	TMap<FGeometryIdentifier, TSharedPtr<FGeometrySelectionTarget>> ActiveTargetMap;
 
 
+	TArray<FGeometryIdentifier> UnlockedTargets;
+
+
 	//
 	// Support for cached FGeometrySelectionTarget / IGeometrySelectors. 
 	// The intention here is to reduce the overhead on selection changes.
@@ -437,6 +446,9 @@ protected:
 
 	friend class FGeometrySelectionManager_SelectionTypeChange;
 	friend class FGeometrySelectionManager_ActiveTargetsChange;
+	
+	friend class FGeometrySelectionManager_TargetLockStateChange;
+	void SetTargetLockStateOnUndoRedo(FGeometryIdentifier TargetIdentifier, bool bLocked);
 
 
 	// support for complex selection changes that are driven externally
diff --git a/Engine/Plugins/Runtime/MeshModelingToolset/Source/ModelingComponents/Public/Selection/GeometrySelector.h b/Engine/Plugins/Runtime/MeshModelingToolset/Source/ModelingComponents/Public/Selection/GeometrySelector.h
--- a/Engine/Plugins/Runtime/MeshModelingToolset/Source/ModelingComponents/Public/Selection/GeometrySelector.h
+++ b/Engine/Plugins/Runtime/MeshModelingToolset/Source/ModelingComponents/Public/Selection/GeometrySelector.h
@@ -62,7 +62,6 @@ struct MODELINGCOMPONENTS_API FGeometryIdentifier
 	 */
 	TWeakObjectPtr<UObject> TargetObject = nullptr;
 
-
 	FGeometryIdentifier() = default;
 	FGeometryIdentifier(const FGeometryIdentifier& Other) = default;
 
@@ -269,6 +268,20 @@ public:
 	virtual FGeometryIdentifier GetIdentifier() const = 0;
 
 
+	//
+	// Locking support. When a Selector is locked, it will not allow selection changes.
+	// The primary use of this is to support configurable enable/disable of selection on 
+	// certain types of objects (eg on static meshes). Although it could potentially be
+	// supported at a higher level, this adds a lot of complications to state management
+	// that are avoided if the Selector is always created.
+	//
+
+	virtual bool IsLockable() const { return false; }
+	virtual bool IsLocked() const { return false; }
+	virtual void SetLockedState(bool bLocked) { }
+
+
+
 	enum class EInitializeSelectionMode
 	{
 		// create a selection of all geometry elements
@@ -450,3 +463,4 @@ public:
 
 
 
+
diff --git a/Engine/Plugins/Runtime/MeshModelingToolset/Source/ModelingComponentsEditorOnly/Private/Selection/StaticMeshSelector.cpp b/Engine/Plugins/Runtime/MeshModelingToolset/Source/ModelingComponentsEditorOnly/Private/Selection/StaticMeshSelector.cpp
--- a/Engine/Plugins/Runtime/MeshModelingToolset/Source/ModelingComponentsEditorOnly/Private/Selection/StaticMeshSelector.cpp
+++ b/Engine/Plugins/Runtime/MeshModelingToolset/Source/ModelingComponentsEditorOnly/Private/Selection/StaticMeshSelector.cpp
@@ -24,6 +24,41 @@ using namespace UE::Geometry;
 
 #define LOCTEXT_NAMESPACE "FStaticMeshSelector"
 
+namespace UEGlobal
+{
+	TSet<UStaticMesh*> UnlockedStaticMeshes;
+}
+
+
+bool FStaticMeshSelector::IsLocked() const
+{
+	return StaticMesh != nullptr && (UEGlobal::UnlockedStaticMeshes.Contains(StaticMesh) == false);
+}
+
+void FStaticMeshSelector::SetLockedState(bool bLocked)
+{
+	if (StaticMesh != nullptr)
+	{
+		if (bLocked)
+		{
+			UEGlobal::UnlockedStaticMeshes.Remove(StaticMesh);
+		}
+		else
+		{
+			UEGlobal::UnlockedStaticMeshes.Add(StaticMesh);
+
+			// if the mesh was locked, we were not updating it on changes, so we
+			// need to do an update now
+			if (IsLocked() == false)
+			{
+				CopyFromStaticMesh();
+				InvalidateOnMeshChange(FDynamicMeshChangeInfo());
+			}
+		}
+	}
+}
+
+
 
 bool FStaticMeshSelector::Initialize(
 	FGeometryIdentifier SourceGeometryIdentifierIn)
@@ -48,14 +83,20 @@ bool FStaticMeshSelector::Initialize(
 	// It is currently a very large hammer, though...
 	StaticMesh_OnMeshChangedHandle = StaticMesh->OnMeshChanged.AddLambda([this]()
 	{
-		CopyFromStaticMesh();
-		InvalidateOnMeshChange(FDynamicMeshChangeInfo());
+		if (IsLocked() == false)
+		{
+			CopyFromStaticMesh();
+			InvalidateOnMeshChange(FDynamicMeshChangeInfo());
+		}
 	});
 
 	// create new transient UDyanmicMesh that will be owned by this Selector
 	LocalTargetMesh.Reset( NewObject<UDynamicMesh>() );
 
-	CopyFromStaticMesh();
+	if (IsLocked() == false)
+	{
+		CopyFromStaticMesh();
+	}
 
 	FBaseDynamicMeshSelector::Initialize(SourceGeometryIdentifierIn, LocalTargetMesh.Get(), 
 		[this]() { return IsValid(this->StaticMeshComponent) ? (UE::Geometry::FTransformSRT3d)this->StaticMeshComponent->GetComponentTransform() : FTransformSRT3d::Identity(); });
diff --git a/Engine/Plugins/Runtime/MeshModelingToolset/Source/ModelingComponentsEditorOnly/Private/Selection/VolumeSelector.cpp b/Engine/Plugins/Runtime/MeshModelingToolset/Source/ModelingComponentsEditorOnly/Private/Selection/VolumeSelector.cpp
--- a/Engine/Plugins/Runtime/MeshModelingToolset/Source/ModelingComponentsEditorOnly/Private/Selection/VolumeSelector.cpp
+++ b/Engine/Plugins/Runtime/MeshModelingToolset/Source/ModelingComponentsEditorOnly/Private/Selection/VolumeSelector.cpp
@@ -18,12 +18,45 @@
 
 #include "Selection/DynamicMeshPolygroupTransformer.h"
 
+#include "Editor.h"		// for GEditor
+
 using namespace UE::Geometry;
 
 #define LOCTEXT_NAMESPACE "FVolumeSelector"
 
 
 
+namespace UEGlobal
+{
+	TSet<UBrushComponent*> UnlockedBrushComponents;
+}
+
+
+bool FVolumeSelector::IsLocked() const
+{
+	return BrushComponent != nullptr && (UEGlobal::UnlockedBrushComponents.Contains(BrushComponent) == false);
+}
+
+void FVolumeSelector::SetLockedState(bool bLocked)
+{
+	if (BrushComponent != nullptr)
+	{
+		if (bLocked)
+		{
+			UEGlobal::UnlockedBrushComponents.Remove(BrushComponent);
+		}
+		else
+		{
+			UEGlobal::UnlockedBrushComponents.Add(BrushComponent);
+			if (IsLocked() == false)
+			{
+				UpdateDynamicMeshFromVolume();
+			}
+		}
+	}
+}
+
+
 
 bool FVolumeSelector::Initialize(
 	FGeometryIdentifier SourceGeometryIdentifierIn)
@@ -45,16 +78,20 @@ bool FVolumeSelector::Initialize(
 	}
 	
 	LocalTargetMesh.Reset( NewObject<UDynamicMesh>() );
-
-	LocalTargetMesh->EditMesh([&](FDynamicMesh3& EditMesh)
-	{
-		UE::Conversion::FVolumeToMeshOptions ConvertOptions;
-		UE::Conversion::BrushComponentToDynamicMesh(BrushComponent, EditMesh, ConvertOptions);
-	});
+	UpdateDynamicMeshFromVolume();
 
 	FBaseDynamicMeshSelector::Initialize(SourceGeometryIdentifierIn, LocalTargetMesh.Get(), 
 		[this]() { return IsValid(this->BrushComponent) ? (UE::Geometry::FTransformSRT3d)this->BrushComponent->GetComponentTransform() : FTransformSRT3d::Identity(); });
 
+	// AVolume and BrushComponent seem to have no way of signalling when they are modified,
+	// so on undo and redo, we have no idea that the mesh has changed. The brute-force way
+	// to work around this is simply to rebuild the mesh on every undo/redo event. 
+	// This is expensive and may have other negative reprecussions...
+	if (GEditor != nullptr)
+	{
+		GEditor->RegisterForUndo(this);
+	}
+
 	return true;
 }
 
@@ -63,6 +100,39 @@ void FVolumeSelector::Shutdown()
 {
 	LocalTargetMesh.Reset();
 	FBaseDynamicMeshSelector::Shutdown();
+
+	if (GEditor != nullptr)
+	{
+		GEditor->UnregisterForUndo(this);
+	}
+}
+
+
+void FVolumeSelector::PostUndo(bool bSuccess)
+{
+	// see comment in Initialize() for details
+	if (IsLocked() == false)
+	{
+		UpdateDynamicMeshFromVolume();
+	}
+}
+void FVolumeSelector::PostRedo(bool bSuccess)
+{
+	// see comment in Initialize() for details
+	if (IsLocked() == false)
+	{
+		UpdateDynamicMeshFromVolume();
+	}
+}
+
+
+void FVolumeSelector::UpdateDynamicMeshFromVolume()
+{
+	LocalTargetMesh->EditMesh([&](FDynamicMesh3& EditMesh)
+	{
+		UE::Conversion::FVolumeToMeshOptions ConvertOptions;
+		UE::Conversion::BrushComponentToDynamicMesh(BrushComponent, EditMesh, ConvertOptions);
+	});
 }
 
 
diff --git a/Engine/Plugins/Runtime/MeshModelingToolset/Source/ModelingComponentsEditorOnly/Public/Selection/StaticMeshSelector.h b/Engine/Plugins/Runtime/MeshModelingToolset/Source/ModelingComponentsEditorOnly/Public/Selection/StaticMeshSelector.h
--- a/Engine/Plugins/Runtime/MeshModelingToolset/Source/ModelingComponentsEditorOnly/Public/Selection/StaticMeshSelector.h
+++ b/Engine/Plugins/Runtime/MeshModelingToolset/Source/ModelingComponentsEditorOnly/Public/Selection/StaticMeshSelector.h
@@ -33,6 +33,10 @@ public:
 	//virtual bool Restore();
 
 
+	virtual bool IsLockable() const override { return true; }
+	virtual bool IsLocked() const override;
+	virtual void SetLockedState(bool bLocked) override;
+
 
 	virtual IGeometrySelectionTransformer* InitializeTransformation(const FGeometrySelection& Selection) override;
 	virtual void ShutdownTransformation(IGeometrySelectionTransformer* Transformer) override;
@@ -54,7 +58,7 @@ protected:
 
 /**
  * FStaticMeshComponentSelectorFactory constructs FStaticMeshSelector instances 
- * for UBrushComponents
+ * for UStaticMeshComponents
  */
 class MODELINGCOMPONENTSEDITORONLY_API FStaticMeshComponentSelectorFactory : public IGeometrySelectorFactory
 {
diff --git a/Engine/Plugins/Runtime/MeshModelingToolset/Source/ModelingComponentsEditorOnly/Public/Selection/VolumeSelector.h b/Engine/Plugins/Runtime/MeshModelingToolset/Source/ModelingComponentsEditorOnly/Public/Selection/VolumeSelector.h
--- a/Engine/Plugins/Runtime/MeshModelingToolset/Source/ModelingComponentsEditorOnly/Public/Selection/VolumeSelector.h
+++ b/Engine/Plugins/Runtime/MeshModelingToolset/Source/ModelingComponentsEditorOnly/Public/Selection/VolumeSelector.h
@@ -7,14 +7,14 @@
 #include "Selection/DynamicMeshSelector.h"
 
 #include "UObject/StrongObjectPtr.h"
-
+#include "EditorUndoClient.h"
 
 class UDynamicMesh;
 class AVolume;
 class UBrushComponent;
 
 
-class MODELINGCOMPONENTSEDITORONLY_API FVolumeSelector : public FBaseDynamicMeshSelector
+class MODELINGCOMPONENTSEDITORONLY_API FVolumeSelector : public FBaseDynamicMeshSelector, public FEditorUndoClient
 {
 public:
 	using FBaseDynamicMeshSelector::Initialize;
@@ -32,10 +32,18 @@ public:
 	//virtual bool Sleep();
 	//virtual bool Restore();
 
+	virtual bool IsLockable() const override { return true; }
+	virtual bool IsLocked() const override;
+	virtual void SetLockedState(bool bLocked) override;
 
 	virtual IGeometrySelectionTransformer* InitializeTransformation(const FGeometrySelection& Selection) override;
 	virtual void ShutdownTransformation(IGeometrySelectionTransformer* Transformer) override;
 
+
+	// FEditorUndoClient implementation
+	virtual void PostUndo(bool bSuccess) override;
+	virtual void PostRedo(bool bSuccess) override;
+
 protected:
 	AVolume* ParentVolume = nullptr;
 	UBrushComponent* BrushComponent = nullptr;
@@ -43,6 +51,8 @@ protected:
 	// TODO: this is not a great design, it would be better if something external could own this mesh...
 	TStrongObjectPtr<UDynamicMesh> LocalTargetMesh;
 
+	void UpdateDynamicMeshFromVolume();
+
 	TSharedPtr<FBasicDynamicMeshSelectionTransformer> ActiveTransformer;
 	void CommitMeshTransform();
 };
