ModelingMode: reorganize tool palettes

diff --git a/Engine/Plugins/Editor/ModelingToolsEditorMode/Source/ModelingToolsEditorMode/Private/ModelingToolsEditorModeToolkit.cpp b/Engine/Plugins/Editor/ModelingToolsEditorMode/Source/ModelingToolsEditorMode/Private/ModelingToolsEditorModeToolkit.cpp
--- a/Engine/Plugins/Editor/ModelingToolsEditorMode/Source/ModelingToolsEditorMode/Private/ModelingToolsEditorModeToolkit.cpp
+++ b/Engine/Plugins/Editor/ModelingToolsEditorMode/Source/ModelingToolsEditorMode/Private/ModelingToolsEditorModeToolkit.cpp
@@ -185,7 +185,8 @@
 	
 	ToolkitBuilder->AddPalette(StaticCastSharedPtr<FEditablePalette>(FavoritesPalette));
 
-	TArray<TSharedPtr<FUICommandInfo>> ShapesPaletteItems({
+
+	TArray<TSharedPtr<FUICommandInfo>> CreatePaletteItems({
 		Commands.BeginAddBoxPrimitiveTool,
 		Commands.BeginAddSpherePrimitiveTool,
 		Commands.BeginAddCylinderPrimitiveTool,
@@ -194,21 +195,12 @@
 		Commands.BeginAddArrowPrimitiveTool,
 		Commands.BeginAddRectanglePrimitiveTool,
 		Commands.BeginAddDiscPrimitiveTool,
-		Commands.BeginAddStairsPrimitiveTool
-	});
-	ToolkitBuilder->AddPalette(
-		MakeShareable( new FToolPalette( Commands.LoadShapesTools.ToSharedRef(), ShapesPaletteItems ) ) );
+		Commands.BeginAddStairsPrimitiveTool,
+		Commands.BeginCubeGridTool,
 
-
-
-	TArray<TSharedPtr<FUICommandInfo>> CreatePaletteItems({
 		Commands.BeginDrawPolygonTool,
 		Commands.BeginDrawPolyPathTool,
-		Commands.BeginDrawAndRevolveTool,
-		Commands.BeginRevolveBoundaryTool,
-		Commands.BeginCombineMeshesTool,
-		Commands.BeginDuplicateMeshesTool,
-		Commands.BeginPatternTool
+		Commands.BeginDrawAndRevolveTool
 	});
 	ToolkitBuilder->AddPalette(
 		MakeShareable( new FToolPalette( Commands.LoadCreateTools.ToSharedRef(), CreatePaletteItems ) ) );
@@ -219,8 +211,8 @@
 			Commands.BeginSelectionAction_Delete,
 			Commands.BeginSelectionAction_Extrude,
 			Commands.BeginSelectionAction_Offset,
-
 			Commands.BeginPolyModelTool_PushPull,
+
 			Commands.BeginPolyModelTool_Inset,
 			Commands.BeginPolyModelTool_Outset,
 			Commands.BeginPolyModelTool_CutFaces,
@@ -240,13 +232,23 @@
 	TArray<TSharedPtr<FUICommandInfo>> TransformPaletteItems({
 		Commands.BeginTransformMeshesTool,
 		Commands.BeginAlignObjectsTool,
+
+		Commands.BeginCombineMeshesTool,
+		Commands.BeginDuplicateMeshesTool,
+
 		Commands.BeginEditPivotTool,
-		Commands.BeginAddPivotActorTool,
 		Commands.BeginBakeTransformTool,
+
 		Commands.BeginTransferMeshTool,
 		Commands.BeginConvertMeshesTool,
-		Commands.BeginSplitMeshesTool
+
+		Commands.BeginSplitMeshesTool,
+		Commands.BeginPatternTool
 	});
+	if (Commands.BeginISMEditorTool)
+	{
+		TransformPaletteItems.Add(Commands.BeginISMEditorTool);
+	}
 	ToolkitBuilder->AddPalette(
 		MakeShareable( new FToolPalette( Commands.LoadTransformTools.ToSharedRef(), TransformPaletteItems ) ) );
 
@@ -259,7 +261,8 @@
 		Commands.BeginOffsetMeshTool,
 		Commands.BeginMeshSpaceDeformerTool,
 		Commands.BeginLatticeDeformerTool,
-		Commands.BeginDisplaceMeshTool
+		Commands.BeginDisplaceMeshTool,
+		Commands.BeginPolyDeformTool
 	});
 	ToolkitBuilder->AddPalette(
 		MakeShareable( new FToolPalette( Commands.LoadDeformTools.ToSharedRef(), DeformPaletteItems ) ) );
@@ -268,12 +271,18 @@
 
 	TArray<TSharedPtr<FUICommandInfo>> PolyToolsPaletteItems({
 		Commands.BeginPolyEditTool,
-		Commands.BeginPolyDeformTool,
-		Commands.BeginCubeGridTool,
+		Commands.BeginSubdividePolyTool,
+
 		Commands.BeginMeshBooleanTool,
+		Commands.BeginPolygonCutTool,
+
+		Commands.BeginPlaneCutTool,
+		Commands.BeginMirrorTool,
+
 		Commands.BeginCutMeshWithMeshTool,
-		Commands.BeginSubdividePolyTool
-	});
+		Commands.BeginMeshTrimTool
+
+		});
 	ToolkitBuilder->AddPalette(
 		MakeShareable( new FToolPalette( Commands.LoadPolyTools.ToSharedRef(), PolyToolsPaletteItems ) ) );
 
@@ -282,27 +291,21 @@
 	TArray<TSharedPtr<FUICommandInfo>> TriToolsPaletteItems({
 		Commands.BeginMeshSelectionTool,
 		Commands.BeginTriEditTool,
+
 		Commands.BeginHoleFillTool,
-		Commands.BeginMirrorTool,
-		Commands.BeginPlaneCutTool,
-		Commands.BeginPolygonCutTool,
-		Commands.BeginMeshTrimTool
-	});
-	ToolkitBuilder->AddPalette(
-		MakeShareable( new FToolPalette( Commands.LoadTriTools.ToSharedRef(), TriToolsPaletteItems ) ) );
-
-
-
-	TArray<TSharedPtr<FUICommandInfo>> MeshOpPaletteItems({
-		Commands.BeginSimplifyMeshTool,
-		Commands.BeginRemeshMeshTool,
 		Commands.BeginWeldEdgesTool,
-		Commands.BeginRemoveOccludedTrianglesTool,
+
 		Commands.BeginSelfUnionTool,
-		Commands.BeginProjectToTargetTool
-	});
+		Commands.BeginRemoveOccludedTrianglesTool,
+
+		Commands.BeginSimplifyMeshTool,
+		Commands.BeginRemeshMeshTool,
+
+		Commands.BeginProjectToTargetTool,
+		Commands.BeginMeshInspectorTool
+		});
 	ToolkitBuilder->AddPalette(
-		MakeShareable( new FToolPalette( Commands.LoadMeshOpsTools.ToSharedRef(), MeshOpPaletteItems ) ) );
+		MakeShareable( new FToolPalette( Commands.LoadMeshOpsTools.ToSharedRef(), TriToolsPaletteItems ) ) );
 
 
 
@@ -329,6 +332,8 @@
 	{
 		BakingPaletteItems.Add(Commands.BeginBakeRenderCaptureTool);
 	}
+	BakingPaletteItems.Add(Commands.BeginMeshInspectorTool);
+	BakingPaletteItems.Add(Commands.BeginEditTangentsTool);
 	ToolkitBuilder->AddPalette(
 		MakeShareable( new FToolPalette( Commands.LoadBakingTools.ToSharedRef(), BakingPaletteItems ) ) );
 
@@ -346,6 +351,8 @@
 	{
 		UVPaletteItems.Add(Commands.LaunchUVEditor);
 	}
+	UVPaletteItems.Add(Commands.BeginMeshInspectorTool);
+	UVPaletteItems.Add(Commands.BeginMeshGroupPaintTool);
 
 	ToolkitBuilder->AddPalette(
 		MakeShareable( new FToolPalette( Commands.LoadUVsTools.ToSharedRef(), UVPaletteItems)) );
@@ -354,50 +361,59 @@
 
 	TArray<TSharedPtr<FUICommandInfo>> AttributesToolsPaletteItems({
 		Commands.BeginMeshInspectorTool,
+		Commands.BeginLODManagerTool,
+
 		Commands.BeginEditNormalsTool,
 		Commands.BeginEditTangentsTool,
-		Commands.BeginAttributeEditorTool,
+
 		Commands.BeginPolyGroupsTool,
 		Commands.BeginMeshGroupPaintTool,
+
+		Commands.BeginAttributeEditorTool,
+		Commands.BeginEditMeshMaterialsTool,
+
 		Commands.BeginMeshAttributePaintTool,
-		Commands.BeginEditMeshMaterialsTool
-	});
-	ToolkitBuilder->AddPalette(
-		MakeShareable( new FToolPalette( Commands.LoadAttributesTools.ToSharedRef(), AttributesToolsPaletteItems ) ) );
-
-
-
-	TArray<TSharedPtr<FUICommandInfo>> VolumePaletteItems({
 		Commands.BeginPhysicsInspectorTool,
 		Commands.BeginSetCollisionGeometryTool,
-		Commands.BeginExtractCollisionGeometryTool,
+		Commands.BeginExtractCollisionGeometryTool
+	});
+	ToolkitBuilder->AddPalette(
+		MakeShareable( new FToolPalette( Commands.LoadAttributesTools.ToSharedRef(), AttributesToolsPaletteItems ) ) );
+
+
+
+	//TArray<TSharedPtr<FUICommandInfo>> VolumePaletteItems({
+
+	//});
+	//ToolkitBuilder->AddPalette(
+	//	MakeShareable( new FToolPalette( Commands.LoadVolumeTools.ToSharedRef(), VolumePaletteItems)) );
+
+
+
+	TArray<TSharedPtr<FUICommandInfo>> LODPaletteItems({
+		Commands.BeginGenerateStaticMeshLODAssetTool,
+		Commands.BeginAddPivotActorTool,
+		Commands.BeginRevolveBoundaryTool,
 		Commands.BeginVolumeToMeshTool,
 		Commands.BeginMeshToVolumeTool
 	});
 	if (Commands.BeginBspConversionTool)
 	{
-		VolumePaletteItems.Add(Commands.BeginBspConversionTool);
+		LODPaletteItems.Add(Commands.BeginBspConversionTool);
 	}
 	ToolkitBuilder->AddPalette(
-		MakeShareable( new FToolPalette( Commands.LoadVolumeTools.ToSharedRef(), VolumePaletteItems)) );
-
-
-
-	TArray<TSharedPtr<FUICommandInfo>> LODPaletteItems({
-		Commands.BeginLODManagerTool,
-		Commands.BeginGenerateStaticMeshLODAssetTool
-	});
-	if (Commands.BeginISMEditorTool)
-	{
-		LODPaletteItems.Add(Commands.BeginISMEditorTool);
-	}
-	ToolkitBuilder->AddPalette(
 		MakeShareable( new FToolPalette( Commands.LoadLodsTools.ToSharedRef(), LODPaletteItems)) );
 
-	ToolkitBuilder->SetActivePaletteOnLoad(Commands.LoadShapesTools.Get());
+	ToolkitBuilder->SetActivePaletteOnLoad(Commands.LoadCreateTools.Get());
 	ToolkitBuilder->UpdateWidget();
 }
 
+
+
+
+
+
+
 void FModelingToolsEditorModeToolkit::Init(const TSharedPtr<IToolkitHost>& InitToolkitHost, TWeakObjectPtr<UEdMode> InOwningMode)
 {
 	const UModelingToolsModeCustomizationSettings* UISettings = GetDefault<UModelingToolsModeCustomizationSettings>();
diff --git a/Engine/Plugins/Editor/ModelingToolsEditorMode/Source/ModelingToolsEditorMode/Private/ModelingToolsManagerActions.cpp b/Engine/Plugins/Editor/ModelingToolsEditorMode/Source/ModelingToolsEditorMode/Private/ModelingToolsManagerActions.cpp
--- a/Engine/Plugins/Editor/ModelingToolsEditorMode/Source/ModelingToolsEditorMode/Private/ModelingToolsManagerActions.cpp
+++ b/Engine/Plugins/Editor/ModelingToolsEditorMode/Source/ModelingToolsEditorMode/Private/ModelingToolsManagerActions.cpp
@@ -190,24 +190,24 @@ void FModelingToolsManagerCommands::RegisterCommands()
 	// VoxOps
 	if ( bUseShortNames )
 	{
-		REGISTER_MODELING_TOOL_COMMAND(BeginVoxelSolidifyTool, "VoxWrap", "Wrap the selected Meshes using voxels");
-		REGISTER_MODELING_TOOL_COMMAND(BeginVoxelBlendTool, "VoxBlnd", "Blend the selected Meshes using voxels");
-		REGISTER_MODELING_TOOL_COMMAND(BeginVoxelMorphologyTool, "VoxMrph", "Offset/Inset the selected Meshes using voxels");
+		REGISTER_MODELING_TOOL_COMMAND(BeginVoxelSolidifyTool, "VoxWrap", "Wrap the selected Meshes using SDFs/voxels");
+		REGISTER_MODELING_TOOL_COMMAND(BeginVoxelBlendTool, "VoxBlnd", "Blend the selected Meshes using SDFs/voxels");
+		REGISTER_MODELING_TOOL_COMMAND(BeginVoxelMorphologyTool, "VoxMrph", "Offset/Inset the selected Meshes using SDFs/voxels");
 #if WITH_PROXYLOD
 		// The ProxyLOD plugin is currently only available on Windows. Without it, the following tools do not work as expected.
-		REGISTER_MODELING_TOOL_COMMAND(BeginVoxelBooleanTool, "VoxBool", "Boolean the selected Meshes using voxels");
-		REGISTER_MODELING_TOOL_COMMAND(BeginVoxelMergeTool, "VoxMrg", "Merge the selected Meshes using voxels");
+		REGISTER_MODELING_TOOL_COMMAND(BeginVoxelBooleanTool, "VoxBool", "Boolean the selected Meshes using SDFs/voxels");
+		REGISTER_MODELING_TOOL_COMMAND(BeginVoxelMergeTool, "VoxMrg", "Merge the selected Meshes using SDFs/voxels");
 #endif	// WITH_PROXYLOD
 	}
 	else
 	{
-		REGISTER_MODELING_TOOL_COMMAND(BeginVoxelSolidifyTool, "Vox Wrap", "Wrap the selected Meshes using voxels");
-		REGISTER_MODELING_TOOL_COMMAND(BeginVoxelBlendTool, "Vox Blend", "Blend the selected Meshes using voxels");
-		REGISTER_MODELING_TOOL_COMMAND(BeginVoxelMorphologyTool, "Vox Offset", "Offset/Inset the selected Meshes using voxels");
+		REGISTER_MODELING_TOOL_COMMAND(BeginVoxelSolidifyTool, "Voxel Wrap", "Wrap the selected Meshes using SDFs/voxels");
+		REGISTER_MODELING_TOOL_COMMAND(BeginVoxelBlendTool, "Voxel Blend", "Blend the selected Meshes using SDFs/voxels");
+		REGISTER_MODELING_TOOL_COMMAND(BeginVoxelMorphologyTool, "Voxel Offset", "Offset/Inset the selected Meshes using SDFs/voxels");
 #if WITH_PROXYLOD
 		// The ProxyLOD plugin is currently only available on Windows. Without it, the following tools do not work as expected.
-		REGISTER_MODELING_TOOL_COMMAND(BeginVoxelBooleanTool, "Vox Boolean", "Boolean the selected Meshes using voxels");
-		REGISTER_MODELING_TOOL_COMMAND(BeginVoxelMergeTool, "Vox Merge", "Merge the selected Meshes using voxels");
+		REGISTER_MODELING_TOOL_COMMAND(BeginVoxelBooleanTool, "Voxel Boolean", "Boolean the selected Meshes using SDFs/voxels");
+		REGISTER_MODELING_TOOL_COMMAND(BeginVoxelMergeTool, "Voxel Merge", "Merge the selected Meshes using SDFs/voxels");
 #endif	// WITH_PROXYLOD
 	}
 
@@ -394,20 +394,20 @@ void FModelingToolsManagerCommands::RegisterCommands()
 	UI_COMMAND(AddToFavorites, "Add to Favorites", "Add to Favorites", EUserInterfaceActionType::Button, FInputChord());
 	UI_COMMAND(RemoveFromFavorites, "Remove from Favorites", "Remove from Favorites", EUserInterfaceActionType::Button, FInputChord());
 	UI_COMMAND(LoadFavoritesTools, "Favorites", "Favorites", EUserInterfaceActionType::RadioButton, FInputChord());
-	UI_COMMAND(LoadSelectionTools, "Selection", "Selection", EUserInterfaceActionType::RadioButton, FInputChord());
+	UI_COMMAND(LoadSelectionTools, "Selection", "Edit Mesh Selections", EUserInterfaceActionType::RadioButton, FInputChord());
 	UI_COMMAND(LoadShapesTools, "Shapes", "Shapes", EUserInterfaceActionType::RadioButton, FInputChord());
-	UI_COMMAND(LoadCreateTools, "Create", "Create", EUserInterfaceActionType::RadioButton, FInputChord());
-	UI_COMMAND(LoadPolyTools, "PolyEdit", "PolyEdit", EUserInterfaceActionType::RadioButton, FInputChord());
-	UI_COMMAND(LoadTriTools, "TriTools", "TriTools", EUserInterfaceActionType::RadioButton, FInputChord());
-	UI_COMMAND(LoadDeformTools, "Deform", "Deform", EUserInterfaceActionType::RadioButton, FInputChord());
-	UI_COMMAND(LoadTransformTools, "Transform", "Transform", EUserInterfaceActionType::RadioButton, FInputChord());
-	UI_COMMAND(LoadMeshOpsTools, "MeshOps", "MeshOps", EUserInterfaceActionType::RadioButton, FInputChord());
-	UI_COMMAND(LoadVoxOpsTools, "VoxOps", "VoxOps", EUserInterfaceActionType::RadioButton, FInputChord());
-	UI_COMMAND(LoadAttributesTools, "Attributes", "Atrribute", EUserInterfaceActionType::RadioButton, FInputChord());
-	UI_COMMAND(LoadUVsTools, "UVs", "UVs", EUserInterfaceActionType::RadioButton, FInputChord());
-	UI_COMMAND(LoadBakingTools, "Baking", "Baking", EUserInterfaceActionType::RadioButton, FInputChord());
+	UI_COMMAND(LoadCreateTools, "Create", "Create New Shapes", EUserInterfaceActionType::RadioButton, FInputChord());
+	UI_COMMAND(LoadPolyTools, "Model", "Shape Modeling", EUserInterfaceActionType::RadioButton, FInputChord());
+	UI_COMMAND(LoadTriTools, "Process", "Mesh Processing", EUserInterfaceActionType::RadioButton, FInputChord());
+	UI_COMMAND(LoadDeformTools, "Deform", "Deformations", EUserInterfaceActionType::RadioButton, FInputChord());
+	UI_COMMAND(LoadTransformTools, "Transform", "Transforms & Conversion", EUserInterfaceActionType::RadioButton, FInputChord());
+	UI_COMMAND(LoadMeshOpsTools, "MeshOps", "Mesh Processing", EUserInterfaceActionType::RadioButton, FInputChord());
+	UI_COMMAND(LoadVoxOpsTools, "VoxOps", "Voxel Processing", EUserInterfaceActionType::RadioButton, FInputChord());
+	UI_COMMAND(LoadAttributesTools, "Attributes", "Mesh Attributes", EUserInterfaceActionType::RadioButton, FInputChord());
+	UI_COMMAND(LoadUVsTools, "UVs", "Create & Edit UVs", EUserInterfaceActionType::RadioButton, FInputChord());
+	UI_COMMAND(LoadBakingTools, "Baking", "Bake Textures & Colors", EUserInterfaceActionType::RadioButton, FInputChord());
 	UI_COMMAND(LoadVolumeTools, "Volumes", "Volumes", EUserInterfaceActionType::RadioButton, FInputChord());
-	UI_COMMAND(LoadLodsTools, "LODs", "LODs", EUserInterfaceActionType::RadioButton, FInputChord());
+	UI_COMMAND(LoadLodsTools, "Misc", "Additional Utility Tools", EUserInterfaceActionType::RadioButton, FInputChord());
 
 	UI_COMMAND(AcceptActiveTool, "Accept", "Accept the active Tool", EUserInterfaceActionType::Button, FInputChord());
 	UI_COMMAND(CancelActiveTool, "Cancel", "Cancel the active Tool", EUserInterfaceActionType::Button, FInputChord());
