ModelingMode: updates to selection system.
Selection enable/disable can now be toggled dynamically via mode settings popup, toolbar will hide/show based on current setting state.
GeometrySelectionManager is now always created. New Mode function GetMeshElementSelectionSystemEnabled() indicates whether selection is currently enabled, when disabled active targets will not be tracked.
Selection Toolbar icons now are disabled while in a Tool (as selection system is not active then)
New Static Mesh Assets and Volumes created using Modeling Tools now default to unlocked, via new functions on StaticMeshSelector / Volume Selector.
Unlocked Asset/Volume sets are cleared on level change
Added cvars modeling.Selection.EnableStaticMeshLocking and modeling.Selection.EnableVolumeLocking which can be used to fully disable lock-ability of StaticMesh/Volume (default true)
shortened some names of tool group tabs, as labels are now shown and long names make the bar wider
Added a few toast notifications for when setting state changes require a mode or editor restart.

diff --git a/Engine/Plugins/Editor/ModelingToolsEditorMode/Source/ModelingToolsEditorMode/Private/ModelingModeToolkit_Toolbars.cpp b/Engine/Plugins/Editor/ModelingToolsEditorMode/Source/ModelingToolsEditorMode/Private/ModelingModeToolkit_Toolbars.cpp
--- a/Engine/Plugins/Editor/ModelingToolsEditorMode/Source/ModelingToolsEditorMode/Private/ModelingModeToolkit_Toolbars.cpp
+++ b/Engine/Plugins/Editor/ModelingToolsEditorMode/Source/ModelingToolsEditorMode/Private/ModelingModeToolkit_Toolbars.cpp
@@ -26,6 +26,10 @@
 #include "Widgets/Layout/SSpacer.h"
 #include "Widgets/Input/SComboButton.h"
 
+// for showing toast notifications
+#include "Widgets/Notifications/SNotificationList.h"
+#include "Framework/Notifications/NotificationManager.h"
+
 #include "ToolMenus.h"
 // to toggle component instance selection
 
@@ -126,8 +130,35 @@
 }
 
 
+
+void MakeSubMenu_SelectionSupport(FModelingToolsEditorModeToolkit* Toolkit, FMenuBuilder& MenuBuilder)
+{
+	UModelingToolsEditorModeSettings* Settings = GetMutableDefault<UModelingToolsEditorModeSettings>();
+
+	// toggle for Combined/Separate Gizmo Mode
+	const FUIAction EnableSelectionsToggle(
+		FExecuteAction::CreateLambda([Toolkit, Settings]
+		{
+			Settings->bEnablePersistentSelections = ! Settings->bEnablePersistentSelections;
+			Settings->SaveConfig();
+			Toolkit->NotifySelectionSystemEnabledStateModified();
+		}),
+		FCanExecuteAction(),
+		FIsActionChecked::CreateLambda([Settings]()
+		{
+			return Settings->bEnablePersistentSelections;
+		}));
+	MenuBuilder.AddMenuEntry(LOCTEXT("MeshElementSelectionToggle", "Mesh Element Selection"), 
+		LOCTEXT("MeshElementSelectionToggle_Tooltip", "Enable support for in-viewport Mesh Element Selection in Modeling Mode"),
+		FSlateIcon(), EnableSelectionsToggle, NAME_None, EUserInterfaceActionType::ToggleButton);
+}
+
+
+
+
 void MakeSubMenu_ModeToggles(FMenuBuilder& MenuBuilder)
 {
+	// add toggle for Instance Selection cvar
 	const FUIAction ToggleInstancesEverywhereAction(
 		FExecuteAction::CreateLambda([]
 		{
@@ -143,10 +174,31 @@
 		}));
 
 	MenuBuilder.AddMenuEntry(
-		LOCTEXT("ToggleInstancesSelection", "Enable Instance Selection"), 
+		LOCTEXT("ToggleInstancesSelection", "Instance Selection"), 
 		LOCTEXT("ToggleInstancesSelection_Tooltip", "Enable/Disable support for direct selection of InstancedStaticMeshComponent Instances (via TypedElements.EnableViewportInstanceSelection cvar)"),
 		FSlateIcon(), ToggleInstancesEverywhereAction, NAME_None, EUserInterfaceActionType::ToggleButton);
 
+	// add toggle for Volume Snapping cvar
+	const FUIAction ToggleVolumeSnappingAction(
+		FExecuteAction::CreateLambda([]
+		{
+			IConsoleVariable* CVar = IConsoleManager::Get().FindConsoleVariable(TEXT("modeling.EnableVolumeSnapping")); 
+			bool CurValue = CVar->GetBool();
+			CVar->Set(CurValue ? false : true);
+		}),
+		FCanExecuteAction(),
+		FIsActionChecked::CreateLambda([]()
+		{
+			IConsoleVariable* CVar = IConsoleManager::Get().FindConsoleVariable(TEXT("modeling.EnableVolumeSnapping"));
+			return (CVar->GetBool());
+		}));
+
+	MenuBuilder.AddMenuEntry(
+		LOCTEXT("ToggleVolumeSnapping", "Volume Snapping"), 
+		LOCTEXT("ToggleVolumeSnapping_Tooltip", "Enable Vertex/Face Snapping and Ray-hits against Volumes in the Level. Note that if your level contains many overlapping large volumes, some Modeling functionality may not work correctly"),
+		FSlateIcon(), ToggleVolumeSnappingAction, NAME_None, EUserInterfaceActionType::ToggleButton);
+
+
 }
 
 
@@ -207,16 +259,26 @@
 
 	MenuBuilder.BeginSection("DefaultObjecTypeSection", LOCTEXT("DefaultObjecTypeSection", "Default Object Type"));
 
+	auto ShowMeshObjectDefaultChangeToast = []()
+	{
+		FNotificationInfo Info(LOCTEXT("ChangeDefaultMeshObjectTypeToast", 
+			"Changing the Default Mesh Object Type will take effect the next time the Editor is started"));
+		Info.ExpireDuration = 5.0f;
+		FSlateNotificationManager::Get().AddNotification(Info);
+	};
+
+
 	MenuBuilder.AddMenuEntry(
 		LOCTEXT("DefaultMeshObjectType_StaticMesh", "Static Mesh"), 
 		TAttribute<FText>(), 
 		FSlateIcon(), 
 		FUIAction(
-			FExecuteAction::CreateLambda([Settings]
+			FExecuteAction::CreateLambda([Settings, ShowMeshObjectDefaultChangeToast]
 			{
 				Settings->DefaultMeshObjectType = EModelingModeDefaultMeshObjectType::StaticMeshAsset;
 				UCreateMeshObjectTypeProperties::DefaultObjectTypeIdentifier = UCreateMeshObjectTypeProperties::StaticMeshIdentifier;
 				Settings->SaveConfig();
+				ShowMeshObjectDefaultChangeToast();
 			}),
 			FCanExecuteAction(),
 			FIsActionChecked::CreateLambda([Settings]()
@@ -230,11 +292,12 @@
 		TAttribute<FText>(), 
 		FSlateIcon(), 
 		FUIAction(
-			FExecuteAction::CreateLambda([Settings]
+			FExecuteAction::CreateLambda([Settings, ShowMeshObjectDefaultChangeToast]
 			{
 				Settings->DefaultMeshObjectType = EModelingModeDefaultMeshObjectType::DynamicMeshActor;
 				UCreateMeshObjectTypeProperties::DefaultObjectTypeIdentifier = UCreateMeshObjectTypeProperties::DynamicMeshActorIdentifier;
 				Settings->SaveConfig();
+				ShowMeshObjectDefaultChangeToast();
 			}),
 			FCanExecuteAction(),
 			FIsActionChecked::CreateLambda([Settings]()
@@ -248,11 +311,12 @@
 		TAttribute<FText>(), 
 		FSlateIcon(), 
 		FUIAction(
-			FExecuteAction::CreateLambda([Settings]
+			FExecuteAction::CreateLambda([Settings, ShowMeshObjectDefaultChangeToast]
 			{
 				Settings->DefaultMeshObjectType = EModelingModeDefaultMeshObjectType::VolumeActor;
 				UCreateMeshObjectTypeProperties::DefaultObjectTypeIdentifier = UCreateMeshObjectTypeProperties::VolumeIdentifier;
 				Settings->SaveConfig();
+				ShowMeshObjectDefaultChangeToast();
 			}),
 			FCanExecuteAction(),
 			FIsActionChecked::CreateLambda([Settings]()
@@ -402,6 +466,10 @@
 	MakeSubMenu_GizmoVisibilityMode(this, MenuBuilder);
 	MenuBuilder.EndSection();
 
+	MenuBuilder.BeginSection("Section_SelectionSupport", LOCTEXT("Section_SelectionSupport", "Mesh Element Selection"));
+	MakeSubMenu_SelectionSupport(this, MenuBuilder);
+	MenuBuilder.EndSection();
+
 	MenuBuilder.BeginSection("Section_MeshObjects", LOCTEXT("Section_MeshObjects", "New Mesh Objects"));
 	MenuBuilder.AddSubMenu(
 		LOCTEXT("MeshObjectTypeSubMenu", "New Mesh Settings"), LOCTEXT("MeshObjectTypeSubMenu_ToolTip", "Configure default settings for new Mesh Object Types"),
@@ -410,7 +478,7 @@
 		}));
 	MenuBuilder.EndSection();
 
-	MenuBuilder.BeginSection("Section_Selection", LOCTEXT("Section_Selection", "Instances"));
+	MenuBuilder.BeginSection("Section_ModeToggles", LOCTEXT("Section_ModeToggles", "Mode Settings"));
 	MakeSubMenu_ModeToggles(MenuBuilder);
 	MenuBuilder.EndSection();
 
@@ -568,6 +636,43 @@
 	ToolbarBuilder.AddWidget(SNew(SSpacer).Size(FVector2D(1, 8)));
 
 
+	ToolbarBuilder.AddWidget(
+		SNew(SComboButton)
+		.HasDownArrow(false)
+		.MenuPlacement(EMenuPlacement::MenuPlacement_MenuRight)
+		.ComboButtonStyle(FAppStyle::Get(), "SimpleComboButton")
+		.IsEnabled_Lambda([this]() { return IsInActiveTool() == false; })
+		.OnGetMenuContent(FOnGetContent::CreateLambda([this]()
+			{
+				return UELocal::MakeMenu_SelectionEdits(this);
+			}))
+		.ContentPadding(FMargin(1.0f, 1.0f))
+				.ButtonContent()
+				[
+					SNew(SImage)
+					.Image(FModelingToolsEditorModeStyle::Get()->GetBrush("ModelingModeSelection.Edits_Right"))
+				]);
+
+	ToolbarBuilder.AddWidget(
+		SNew(SComboButton)
+		.HasDownArrow(false)
+		.MenuPlacement(EMenuPlacement::MenuPlacement_MenuRight)
+		.ComboButtonStyle(FAppStyle::Get(), "SimpleComboButton")
+		.IsEnabled_Lambda([this]() { return IsInActiveTool() == false; })
+		.OnGetMenuContent(FOnGetContent::CreateLambda([this]()
+			{
+				return UELocal::MakeMenu_SelectionConfigSettings(this);
+			}))
+		.ContentPadding(FMargin(1.0f, 1.0f))
+				.ButtonContent()
+				[
+					SNew(SImage)
+					.Image(FModelingToolsEditorModeStyle::Get()->GetBrush("ModelingMode.DefaultSettings"))
+					.ColorAndOpacity(FSlateColor::UseForeground())
+				]);
+
+	ToolbarBuilder.AddWidget(SNew(SSpacer).Size(FVector2D(1, 16)));
+
 	//
 	// locking toggle button is implemented using 3 separate buttons because we cannot 
 	// control enabled/disabled or background color for a toolbar button
@@ -575,32 +680,32 @@
 
 	FUIAction UnlockTargetAction(
 		FExecuteAction::CreateLambda([this] { GetMeshSelectionManager()->SetCurrentTargetsLockState(false); }),
-		FCanExecuteAction::CreateLambda([this] { return GetMeshSelectionManager()->GetAnyCurrentTargetsLocked(); }),
+		FCanExecuteAction::CreateLambda([this] { return GetMeshSelectionManager()->GetAnyCurrentTargetsLocked() && IsInActiveTool() == false; }),
 		FIsActionChecked::CreateLambda([] { return false; }),
 		FIsActionButtonVisible::CreateLambda([this] {
 			return GetMeshSelectionManager()->GetMeshTopologyMode() != UGeometrySelectionManager::EMeshTopologyMode::None
-				&& GetMeshSelectionManager()->GetAnyCurrentTargetsLockable() 
-				&& GetMeshSelectionManager()->GetAnyCurrentTargetsLocked();
-		})
+			&& GetMeshSelectionManager()->GetAnyCurrentTargetsLockable()
+			&& GetMeshSelectionManager()->GetAnyCurrentTargetsLocked();
+			})
 	);
 	FUIAction LockTargetAction(
 		FExecuteAction::CreateLambda([this] { GetMeshSelectionManager()->SetCurrentTargetsLockState(true); }),
-		FCanExecuteAction::CreateLambda([this] { return !GetMeshSelectionManager()->GetAnyCurrentTargetsLocked(); }),
+		FCanExecuteAction::CreateLambda([this] { return !GetMeshSelectionManager()->GetAnyCurrentTargetsLocked() && IsInActiveTool() == false; }),
 		FIsActionChecked::CreateLambda([]() { return false; }),
 		FIsActionButtonVisible::CreateLambda([this] {
 			return GetMeshSelectionManager()->GetMeshTopologyMode() != UGeometrySelectionManager::EMeshTopologyMode::None
-				&& GetMeshSelectionManager()->GetAnyCurrentTargetsLockable() 
-				&& (GetMeshSelectionManager()->GetAnyCurrentTargetsLocked() == false);
-		})
+			&& GetMeshSelectionManager()->GetAnyCurrentTargetsLockable()
+			&& (GetMeshSelectionManager()->GetAnyCurrentTargetsLocked() == false);
+			})
 	);
 	FUIAction DisabledLockTargetAction(
-		FExecuteAction::CreateLambda([]{ }),
+		FExecuteAction::CreateLambda([] {}),
 		FCanExecuteAction::CreateLambda([] { return false; }),
 		FIsActionChecked::CreateLambda([]() { return false; }),
-		FIsActionButtonVisible::CreateLambda([this] { 
+		FIsActionButtonVisible::CreateLambda([this] {
 			return GetMeshSelectionManager()->GetMeshTopologyMode() == UGeometrySelectionManager::EMeshTopologyMode::None
-				|| GetMeshSelectionManager()->GetAnyCurrentTargetsLockable() == false; 
-		})
+			|| GetMeshSelectionManager()->GetAnyCurrentTargetsLockable() == false;
+			})
 	);
 
 	ToolbarBuilder.BeginStyleOverride(FName("SelectionToolBar.RedButton"));
@@ -620,40 +725,7 @@
 		LOCTEXT("Selection_LockTargetDisabled_Tooltip", "No Active Selection Targets are Lockable"),
 		FSlateIcon(FModelingToolsEditorModeStyle::GetStyleSetName(), "SelectionToolBarIcons.UnlockedTarget"));
 
-	ToolbarBuilder.AddWidget(SNew(SSpacer).Size(FVector2D(1, 8)));
-
-	ToolbarBuilder.AddWidget(
-		SNew(SComboButton)
-		.HasDownArrow(false)
-		.MenuPlacement(EMenuPlacement::MenuPlacement_MenuRight)
-		.ComboButtonStyle(FAppStyle::Get(), "SimpleComboButton")
-		.OnGetMenuContent(FOnGetContent::CreateLambda([this]()
-			{
-				return UELocal::MakeMenu_SelectionEdits(this);
-			}))
-		.ContentPadding(FMargin(1.0f, 1.0f))
-				.ButtonContent()
-				[
-					SNew(SImage)
-					.Image(FModelingToolsEditorModeStyle::Get()->GetBrush("ModelingModeSelection.Edits_Right"))
-				]);
-
-	ToolbarBuilder.AddWidget(
-		SNew(SComboButton)
-		.HasDownArrow(false)
-		.MenuPlacement(EMenuPlacement::MenuPlacement_MenuRight)
-		.ComboButtonStyle(FAppStyle::Get(), "SimpleComboButton")
-		.OnGetMenuContent(FOnGetContent::CreateLambda([this]()
-			{
-				return UELocal::MakeMenu_SelectionConfigSettings(this);
-			}))
-		.ContentPadding(FMargin(1.0f, 1.0f))
-				.ButtonContent()
-				[
-					SNew(SImage)
-					.Image(FModelingToolsEditorModeStyle::Get()->GetBrush("ModelingMode.DefaultSettings"))
-					.ColorAndOpacity(FSlateColor::UseForeground())
-				]);
 
 
+
 	TSharedRef<SWidget> Toolbar = ToolbarBuilder.MakeWidget();
@@ -679,6 +751,16 @@
 		]
 		];
 
+	SelectionPaletteOverlayWidget->SetVisibility( TAttribute<EVisibility>::CreateLambda([this]()
+	{
+		if ( UModelingToolsEditorMode* ModelingMode = Cast<UModelingToolsEditorMode>(GetScriptableEditorMode()) )
+		{
+			return ModelingMode->GetMeshElementSelectionSystemEnabled() ? EVisibility::Visible : EVisibility::Collapsed;
+		}
+		return EVisibility::Collapsed;
+	}) );
+
+
 }
 
 
diff --git a/Engine/Plugins/Editor/ModelingToolsEditorMode/Source/ModelingToolsEditorMode/Private/ModelingToolsActions.cpp b/Engine/Plugins/Editor/ModelingToolsEditorMode/Source/ModelingToolsEditorMode/Private/ModelingToolsActions.cpp
--- a/Engine/Plugins/Editor/ModelingToolsEditorMode/Source/ModelingToolsEditorMode/Private/ModelingToolsActions.cpp
+++ b/Engine/Plugins/Editor/ModelingToolsEditorMode/Source/ModelingToolsEditorMode/Private/ModelingToolsActions.cpp
@@ -37,6 +37,7 @@ FModelingModeActionCommands::FModelingModeActionCommands() :
 void FModelingModeActionCommands::RegisterCommands()
 {
 	UI_COMMAND(FocusViewCommand, "Focus View at Cursor", "Focuses the camera at the scene hit location under the cursor", EUserInterfaceActionType::None, FInputChord(EKeys::C));
+	UI_COMMAND(ToggleSelectionLockStateCommand, "Toggle Selection Lock State", "Toggles the Locked/Unlocked state of the active Selection Target", EUserInterfaceActionType::None, FInputChord(EKeys::U));
 }
 
 
@@ -47,12 +48,17 @@ void FModelingModeActionCommands::RegisterCommandBindings(TSharedPtr<FUICommandL
 	UICommandList->MapAction(
 		Commands.FocusViewCommand,
 		FExecuteAction::CreateLambda([OnCommandExecuted]() { OnCommandExecuted(EModelingModeActionCommands::FocusViewToCursor); }));
+
+	UICommandList->MapAction(
+		Commands.ToggleSelectionLockStateCommand,
+		FExecuteAction::CreateLambda([OnCommandExecuted]() { OnCommandExecuted(EModelingModeActionCommands::ToggleSelectionLockState); }));
 }
 
 void FModelingModeActionCommands::UnRegisterCommandBindings(TSharedPtr<FUICommandList> UICommandList)
 {
 	const FModelingModeActionCommands& Commands = FModelingModeActionCommands::Get();
 	UICommandList->UnmapAction(Commands.FocusViewCommand);
+	UICommandList->UnmapAction(Commands.ToggleSelectionLockStateCommand);
 }
 
 
diff --git a/Engine/Plugins/Editor/ModelingToolsEditorMode/Source/ModelingToolsEditorMode/Private/ModelingToolsEditorMode.cpp b/Engine/Plugins/Editor/ModelingToolsEditorMode/Source/ModelingToolsEditorMode/Private/ModelingToolsEditorMode.cpp
--- a/Engine/Plugins/Editor/ModelingToolsEditorMode/Source/ModelingToolsEditorMode/Private/ModelingToolsEditorMode.cpp
+++ b/Engine/Plugins/Editor/ModelingToolsEditorMode/Source/ModelingToolsEditorMode/Private/ModelingToolsEditorMode.cpp
@@ -148,6 +148,8 @@ const FEditorModeID UModelingToolsEditorMode::EM_ModelingToolsEditorModeId = TEX
 FDateTime UModelingToolsEditorMode::LastModeStartTimestamp;
 FDateTime UModelingToolsEditorMode::LastToolStartTimestamp;
 
+FDelegateHandle UModelingToolsEditorMode::GlobalModelingWorldTeardownEventHandle;
+
 namespace
 {
 FString GetToolName(const UInteractiveTool& Tool)
@@ -178,11 +180,11 @@ UModelingToolsEditorMode::~UModelingToolsEditorMode()
 
 bool UModelingToolsEditorMode::ProcessEditDelete()
 {
-	if (SelectionManager && SelectionManager->HasSelection())
+	if (GetSelectionManager() && GetSelectionManager()->HasSelection())
 	{
 		if ( UDeleteGeometrySelectionCommand* DeleteCommand = Cast<UDeleteGeometrySelectionCommand>(ModelingModeCommands[0]) )
 		{
-			SelectionManager->ExecuteSelectionCommand(DeleteCommand);
+			GetSelectionManager()->ExecuteSelectionCommand(DeleteCommand);
 			return true;
 		}
 	}
@@ -358,6 +360,7 @@ void UModelingToolsEditorMode::Enter()
 
 	const UModelingToolsEditorModeSettings* ModelingModeSettings = GetDefault<UModelingToolsEditorModeSettings>();
 	const UModelingToolsModeCustomizationSettings* ModelingEditorSettings = GetDefault<UModelingToolsModeCustomizationSettings>();
+	bSelectionSystemEnabled = ModelingModeSettings->bEnablePersistentSelections;
 
 	// Register builders for tool targets that the mode uses.
 	GetInteractiveToolsContext()->TargetManager->AddTargetFactory(NewObject<UStaticMeshComponentToolTargetFactory>(GetToolManager()));
@@ -378,6 +381,10 @@ void UModelingToolsEditorMode::Enter()
 		return true;
 	});
 
+	// register for OnRender and OnDrawHUD extensions
+	GetInteractiveToolsContext()->OnRender.AddUObject(this, &UModelingToolsEditorMode::OnToolsContextRender);
+	GetInteractiveToolsContext()->OnDrawHUD.AddUObject(this, &UModelingToolsEditorMode::OnToolsContextDrawHUD);
+
 	// register stylus event handler
 	StylusStateTracker = MakeUnique<FStylusStateTracker>();
 
@@ -400,50 +407,46 @@ void UModelingToolsEditorMode::Enter()
 	UE::Geometry::RegisterSceneSnappingManager(GetInteractiveToolsContext());
 	SceneSnappingManager = UE::Geometry::FindModelingSceneSnappingManager(GetToolManager());
 
-	// register selection manager, if this feature is enabled in the mode settings
-	if (ModelingModeSettings && ModelingModeSettings->GetMeshSelectionsEnabled())
-	{
-		// set up SelectionManager and register known factory types
-		SelectionManager = NewObject<UGeometrySelectionManager>(GetToolManager());
-		SelectionManager->Initialize(GetInteractiveToolsContext(), GetToolManager()->GetContextTransactionsAPI());
-		SelectionManager->RegisterSelectorFactory(MakeUnique<FDynamicMeshComponentSelectorFactory>());
-		SelectionManager->RegisterSelectorFactory(MakeUnique<FBrushComponentSelectorFactory>());
-		SelectionManager->RegisterSelectorFactory(MakeUnique<FStaticMeshComponentSelectorFactory>());
 
-		GetInteractiveToolsContext()->OnRender.AddUObject(this, &UModelingToolsEditorMode::OnToolsContextRender);
-		GetInteractiveToolsContext()->OnDrawHUD.AddUObject(this, &UModelingToolsEditorMode::OnToolsContextDrawHUD);
-		// this is hopefully temporary? kinda gross...
-		GetInteractiveToolsContext()->ContextObjectStore->AddContextObject(SelectionManager);
-
-		// rebuild tool palette on any selection changes. This is expensive and ideally will be
-		// optimized in the future.
-		SelectionManager_SelectionModifiedHandle = SelectionManager->OnSelectionModified.AddLambda([this]()
-		{
-			((FModelingToolsEditorModeToolkit*)Toolkit.Get())->ForceToolPaletteRebuild();
-		});
+	// set up SelectionManager and register known factory types
+	SelectionManager = NewObject<UGeometrySelectionManager>(GetToolManager());
+	SelectionManager->Initialize(GetInteractiveToolsContext(), GetToolManager()->GetContextTransactionsAPI());
+	SelectionManager->RegisterSelectorFactory(MakeUnique<FDynamicMeshComponentSelectorFactory>());
+	SelectionManager->RegisterSelectorFactory(MakeUnique<FBrushComponentSelectorFactory>());
+	SelectionManager->RegisterSelectorFactory(MakeUnique<FStaticMeshComponentSelectorFactory>());
+
+	// this is hopefully temporary? kinda gross...
+	GetInteractiveToolsContext()->ContextObjectStore->AddContextObject(SelectionManager);
+
+	// rebuild tool palette on any selection changes. This is expensive and ideally will be
+	// optimized in the future.
+	//SelectionManager_SelectionModifiedHandle = SelectionManager->OnSelectionModified.AddLambda([this]()
+	//{
+	//	((FModelingToolsEditorModeToolkit*)Toolkit.Get())->ForceToolPaletteRebuild();
+	//});
+
+	// set up the selection interaction
+	SelectionInteraction = NewObject<UModelingSelectionInteraction>(GetToolManager());
+	SelectionInteraction->Initialize(SelectionManager,
+		[this]() { return GetGeometrySelectionChangesAllowed(); },
+		[this](const FInputDeviceRay& DeviceRay) { return TestForEditorGizmoHit(DeviceRay); });
+	GetInteractiveToolsContext()->InputRouter->RegisterSource(SelectionInteraction);
+
+	SelectionInteraction->OnTransformBegin.AddLambda([this]() 
+	{ 
+		// Disable the SnappingManager while the SelectionInteraction is editing a mesh via transform gizmo
+		SceneSnappingManager->PauseSceneGeometryUpdates();
+
+		// If the transform is happening via the gizmo numerical UI, then we can run into the same slate
+		// throttling issues as tools. We need to continue receiving render/tick while user scrubs the slate values.
+		FSlateThrottleManager::Get().DisableThrottle(true);
+	});
+	SelectionInteraction->OnTransformEnd.AddLambda([this]() 
+	{ 
+		FSlateThrottleManager::Get().DisableThrottle(false);
+		SceneSnappingManager->UnPauseSceneGeometryUpdates(); 
+	});
 
-		// set up the selection interaction
-		SelectionInteraction = NewObject<UModelingSelectionInteraction>(GetToolManager());
-		SelectionInteraction->Initialize(SelectionManager,
-			[this]() { return GetGeometrySelectionChangesAllowed(); },
-			[this](const FInputDeviceRay& DeviceRay) { return TestForEditorGizmoHit(DeviceRay); });
-		GetInteractiveToolsContext()->InputRouter->RegisterSource(SelectionInteraction);
-
-		SelectionInteraction->OnTransformBegin.AddLambda([this]() 
-		{ 
-			// Disable the SnappingManager while the SelectionInteraction is editing a mesh via transform gizmo
-			SceneSnappingManager->PauseSceneGeometryUpdates();
-
-			// If the transform is happening via the gizmo numerical UI, then we can run into the same slate
-			// throttling issues as tools. We need to continue receiving render/tick while user scrubs the slate values.
-			FSlateThrottleManager::Get().DisableThrottle(true);
-		});
-		SelectionInteraction->OnTransformEnd.AddLambda([this]() 
-		{ 
-			FSlateThrottleManager::Get().DisableThrottle(false);
-			SceneSnappingManager->UnPauseSceneGeometryUpdates(); 
-		});
-	}
 
 	// register level objects observer that will update the snapping manager as the scene changes
 	LevelObjectsObserver = MakeShared<FLevelObjectsObserver>();
@@ -480,6 +483,17 @@ void UModelingToolsEditorMode::Enter()
 			if (CreatedInfo.NewAsset != nullptr)
 			{
 				UE::Modeling::OnNewAssetCreated(CreatedInfo.NewAsset);
+				// If we are creating a new asset or component, it should be initially unlocked in the Selection system.
+				// Currently have no generic way to do this, the Selection Manager does not necessarily support Static Meshes
+				// or Brush Components. So doing it here...
+				if (UStaticMesh* StaticMesh = Cast<UStaticMesh>(CreatedInfo.NewAsset))
+				{
+					FStaticMeshSelector::SetAssetUnlockedOnCreation(StaticMesh);
+				}
+			}
+			if ( UBrushComponent* BrushComponent = Cast<UBrushComponent>(CreatedInfo.NewComponent) )
+			{
+				FVolumeSelector::SetComponentUnlockedOnCreation(BrushComponent);
 			}
 		});
 		TextureCreatedEventHandle = ModelCreationAPI->OnModelingTextureCreated.AddLambda([](const FCreateTextureObjectResult& CreatedInfo)
@@ -816,7 +830,7 @@ void UModelingToolsEditorMode::Enter()
 				ModelingEditorSettings->LastMeshSelectionElementType = static_cast<int>(ElementMode);
 				ModelingEditorSettings->SaveConfig();
 			}),
-			FCanExecuteAction(),
+			FCanExecuteAction::CreateLambda([this]() { return GetToolManager()->HasAnyActiveTool() == false; }),
 			FIsActionChecked::CreateLambda([this, TopoMode, ElementMode]() { return GetSelectionManager()->GetMeshTopologyMode() == TopoMode && GetSelectionManager()->GetSelectionElementType() == ElementMode; }),
 			EUIActionRepeatMode::RepeatDisabled);
 	};
@@ -959,17 +973,17 @@ void UModelingToolsEditorMode::Enter()
 		[this](const UObject* Object)  { UpdateSelectionManagerOnEditorSelectionChange(); } );
 
 	// restore various settings
-	if (SelectionManager)
+	if (GetSelectionManager())
 	{
 		UE::Geometry::EGeometryElementType LastElementType = static_cast<UE::Geometry::EGeometryElementType>(ModelingEditorSettings->LastMeshSelectionElementType);
 		if (LastElementType == UE::Geometry::EGeometryElementType::Edge || LastElementType == UE::Geometry::EGeometryElementType::Face || LastElementType == UE::Geometry::EGeometryElementType::Vertex)
 		{
-			SelectionManager->SetSelectionElementType(LastElementType);
+			GetSelectionManager()->SetSelectionElementType(LastElementType);
 		}
 		UGeometrySelectionManager::EMeshTopologyMode LastTopologyMode = static_cast<UGeometrySelectionManager::EMeshTopologyMode>(ModelingEditorSettings->LastMeshSelectionTopologyMode);
 		if (LastTopologyMode == UGeometrySelectionManager::EMeshTopologyMode::None || LastTopologyMode == UGeometrySelectionManager::EMeshTopologyMode::Triangle || LastTopologyMode == UGeometrySelectionManager::EMeshTopologyMode::Polygroup)
 		{
-			SelectionManager->SetMeshTopologyMode(LastTopologyMode);
+			GetSelectionManager()->SetMeshTopologyMode(LastTopologyMode);
 		}
 
 		bEnableStaticMeshElementSelection = ModelingEditorSettings->bLastMeshSelectionStaticMeshToggle;
@@ -984,9 +998,25 @@ void UModelingToolsEditorMode::Enter()
 		}
 	}
 
-
 	// initialize SelectionManager w/ active selection
 	UpdateSelectionManagerOnEditorSelectionChange(true);
+
+	// Selection system currently requires the concept of 'locking' for Static Meshes and Volumes. This is maintained
+	// by a global list that we do *not* want to clear between invocations of Modeling Mode (v annoying if frequently switching
+	// modes) but *do* want to clear when the user loads a new level. So, the first time this runs, register a delegate that listens
+	// for level editor map changes. This is a static member and will never be unregistered!
+	if (GlobalModelingWorldTeardownEventHandle.IsValid() == false)
+	{
+		FLevelEditorModule& LevelEditor = FModuleManager::LoadModuleChecked<FLevelEditorModule>("LevelEditor");
+		GlobalModelingWorldTeardownEventHandle = LevelEditor.OnMapChanged().AddLambda([](UWorld* World, EMapChangeType ChangeType)
+		{
+			if (ChangeType == EMapChangeType::TearDownWorld)
+			{
+				FVolumeSelector::ResetUnlockedBrushComponents();
+				FStaticMeshSelector::ResetUnlockedStaticMeshAssets();
+			}
+		});
+	}
 }
 
 void UModelingToolsEditorMode::RegisterUVEditor()
@@ -1178,10 +1208,10 @@ void UModelingToolsEditorMode::OnToolsContextRender(IToolsContextRenderAPI* Rend
 		SelectionInteraction->ApplyPendingTransformInteractions();
 	}
 
-	if (SelectionManager)
+	if (GetSelectionManager())
 	{
 		// currently relying on debug rendering to visualize selections
-		SelectionManager->DebugRender(RenderAPI);
+		GetSelectionManager()->DebugRender(RenderAPI);
 	}
 }
 
@@ -1212,6 +1242,12 @@ bool UModelingToolsEditorMode::ShouldToolStartBeAllowed(const FString& ToolIdent
 
 bool UModelingToolsEditorMode::GetGeometrySelectionChangesAllowed() const
 {
+	// disable selection system if it is...disabled
+	if (GetMeshElementSelectionSystemEnabled() == false)
+	{
+		return false;
+	}
+
 	// disable selection system if we are in a Tool
 	if ( GetToolManager() && GetToolManager()->HasAnyActiveTool() )
 	{
@@ -1237,10 +1273,45 @@ bool UModelingToolsEditorMode::TestForEditorGizmoHit(const FInputDeviceRay& Clic
 }
 
 
+bool UModelingToolsEditorMode::GetMeshElementSelectionSystemEnabled() const
+{
+	return bSelectionSystemEnabled;
+}
+
+void UModelingToolsEditorMode::NotifySelectionSystemEnabledStateModified()
+{
+	UModelingToolsEditorModeSettings* Settings = GetMutableDefault<UModelingToolsEditorModeSettings>();
+	bool bNewState = Settings->bEnablePersistentSelections;
+	if (bNewState != bSelectionSystemEnabled)
+	{
+		if (bNewState == true)
+		{
+			UpdateSelectionManagerOnEditorSelectionChange(true);		// do like a mode enter so that we get an undoable active-target state
+		}
+		else
+		{
+			if (SelectionManager && (SelectionManager->HasSelection() || SelectionManager->HasActiveTargets()))
+			{
+				GetToolManager()->GetContextTransactionsAPI()->BeginUndoTransaction(LOCTEXT("InitializeSelection", "Initialize Selection"));
+
+				SelectionManager->SynchronizeActiveTargets(TArray<FGeometryIdentifier>(),
+					[this]() { SelectionManager->ClearSelection(); });
+
+				GetToolManager()->GetContextTransactionsAPI()->EndUndoTransaction();
+			}
+		}
+
+		// update things
+
+		bSelectionSystemEnabled = bNewState;
+		GetInteractiveToolsContext()->PostInvalidation();
+	}
+}
+
 
 void UModelingToolsEditorMode::UpdateSelectionManagerOnEditorSelectionChange(bool bEnteringMode)
 {
-	if (!SelectionManager) return;
+	if (GetMeshElementSelectionSystemEnabled() == false || SelectionManager == nullptr) return;
 
 	// if we are in undo/redo, ignore selection change notifications, the required
 	// changes are handled via FChanges that SelectionManager has emitted
@@ -1314,7 +1385,8 @@ bool UModelingToolsEditorMode::BoxSelect(FBox& InBox, bool InSelect)
 
 bool UModelingToolsEditorMode::FrustumSelect(const FConvexVolume& InFrustum, FEditorViewportClient* InViewportClient, bool InSelect)
 {
-	if (SelectionManager 
+	if (GetMeshElementSelectionSystemEnabled()
+		&& SelectionManager
 		&& SelectionManager->HasActiveTargets() 
 		&& SelectionManager->GetMeshTopologyMode() != UGeometrySelectionManager::EMeshTopologyMode::None)
 	{
@@ -1498,6 +1570,20 @@ void UModelingToolsEditorMode::ModelingModeShortcutRequested(EModelingModeAction
 	{
 		FocusCameraAtCursorHotkey();
 	}
+	else if (Command == EModelingModeActionCommands::ToggleSelectionLockState)
+	{
+		if (SelectionManager)
+		{
+			if (SelectionManager->GetAnyCurrentTargetsLocked())
+			{
+				SelectionManager->SetCurrentTargetsLockState(false);
+			}
+			else
+			{
+				SelectionManager->SetCurrentTargetsLockState(true);
+			}
+		}
+	}
 }
 
 
diff --git a/Engine/Plugins/Editor/ModelingToolsEditorMode/Source/ModelingToolsEditorMode/Private/ModelingToolsEditorModeToolkit.cpp b/Engine/Plugins/Editor/ModelingToolsEditorMode/Source/ModelingToolsEditorMode/Private/ModelingToolsEditorModeToolkit.cpp
--- a/Engine/Plugins/Editor/ModelingToolsEditorMode/Source/ModelingToolsEditorMode/Private/ModelingToolsEditorModeToolkit.cpp
+++ b/Engine/Plugins/Editor/ModelingToolsEditorMode/Source/ModelingToolsEditorMode/Private/ModelingToolsEditorModeToolkit.cpp
@@ -42,6 +42,10 @@
 #include "ModelingModeAssetUtils.h"
 #include "ModelingToolsEditablePaletteConfig.h"
 
+// for showing toast notifications
+#include "Widgets/Notifications/SNotificationList.h"
+#include "Framework/Notifications/NotificationManager.h"
+
 // for presets
 #include "PresetAsset.h"
 #include "PropertyCustomizationHelpers.h"
@@ -217,10 +221,11 @@ void FModelingToolsEditorModeToolkit::RegisterPalettes()
 	UISettings->OnSettingChanged().AddSP(SharedThis(this), &FModelingToolsEditorModeToolkit::UpdateCategoryButtonLabelVisibility);
 	
 	ToolkitSections = MakeShared<FToolkitSections>();
-	ToolkitBuilder = MakeShared<FToolkitBuilder>(
-				ScriptableMode->GetModeInfo().ToolbarCustomizationName, 
-				GetToolkitCommands(),
-				ToolkitSections);
+	FToolkitBuilderArgs ToolkitBuilderArgs(ScriptableMode->GetModeInfo().ToolbarCustomizationName);
+	ToolkitBuilderArgs.ToolkitCommandList = GetToolkitCommands();
+	ToolkitBuilderArgs.ToolkitSections = ToolkitSections;
+	ToolkitBuilderArgs.SelectedCategoryTitleVisibility = EVisibility::Collapsed;
+	ToolkitBuilder = MakeShared<FToolkitBuilder>(ToolkitBuilderArgs);
 	ToolkitBuilder->SetCategoryButtonLabelVisibility(UISettings->bShowCategoryButtonLabels);
 
 	// We need to specify the modeling mode specific config instance because this cannot exist at the base toolkit level due to dependency issues currently
@@ -582,13 +587,9 @@ void FModelingToolsEditorModeToolkit::Init(const TSharedPtr<IToolkitHost>& InitT
 	MakeGizmoNumericalUIOverlayWidget();
 	GetToolkitHost()->AddViewportOverlayWidget(GizmoNumericalUIOverlayWidget.ToSharedRef());
 
-	const UModelingToolsEditorModeSettings* ModelingModeSettings = GetDefault<UModelingToolsEditorModeSettings>();
-	bool bEnableSelectionUI = ModelingModeSettings && ModelingModeSettings->bEnablePersistentSelections;
-	if (bEnableSelectionUI)
-	{
-		MakeSelectionPaletteOverlayWidget();
-		GetToolkitHost()->AddViewportOverlayWidget(SelectionPaletteOverlayWidget.ToSharedRef());
-	}
+	MakeSelectionPaletteOverlayWidget();
+	GetToolkitHost()->AddViewportOverlayWidget(SelectionPaletteOverlayWidget.ToSharedRef());
+}
 
 	CurrentPreset = FSoftObjectPath(); // Default to the default collection by leaving this null.
 }
@@ -2031,6 +2032,8 @@ void FModelingToolsEditorModeToolkit::ShowRealtimeAndModeWarnings(bool bShowReal
 
 void FModelingToolsEditorModeToolkit::OnToolStarted(UInteractiveToolManager* Manager, UInteractiveTool* Tool)
 {
+	bInActiveTool = true;
+
 	UpdateActiveToolProperties();
 
 	UInteractiveTool* CurTool = GetScriptableEditorMode()->GetToolManager(EToolsContextScope::EdMode)->GetActiveTool(EToolSide::Left);
@@ -2072,6 +2075,8 @@ void FModelingToolsEditorModeToolkit::OnToolStarted(UInteractiveToolManager* Man
 
 void FModelingToolsEditorModeToolkit::OnToolEnded(UInteractiveToolManager* Manager, UInteractiveTool* Tool)
 {
+	bInActiveTool = false;
+
 	if (IsHosted())
 	{
 		GetToolkitHost()->RemoveViewportOverlayWidget(ToolShutdownViewportOverlayWidget.ToSharedRef());
@@ -2287,5 +2292,27 @@ void FModelingToolsEditorModeToolkit::OnRestrictiveModeAutoGeneratedAssetPathTex
 	Settings->SetRestrictiveModeAutoGeneratedAssetPath(InNewText.ToString());
 }
 
+void FModelingToolsEditorModeToolkit::NotifySelectionSystemEnabledStateModified()
+{
+	if (UModelingToolsEditorMode* ModelingMode = Cast<UModelingToolsEditorMode>(GetScriptableEditorMode()))
+	{
+		ModelingMode->NotifySelectionSystemEnabledStateModified();
+	}
+	
+	//if (bUsesToolkitBuilder && ToolkitBuilder.IsValid())
+	//{
+	//	ToolkitBuilder->InitCategoryToolbarContainerWidget();
+	//}
+	//else
+	//{
+		// required due to above
+		FNotificationInfo Info(LOCTEXT("ChangeSelectionEnabledToast",
+			"You should exit and re-enter Modeling Mode after toggling Mesh Element Selection Support to fully update the UI"));
+		Info.ExpireDuration = 5.0f;
+		FSlateNotificationManager::Get().AddNotification(Info);
+	//}
+}
+
+
 
 #undef LOCTEXT_NAMESPACE
diff --git a/Engine/Plugins/Editor/ModelingToolsEditorMode/Source/ModelingToolsEditorMode/Private/ModelingToolsManagerActions.cpp b/Engine/Plugins/Editor/ModelingToolsEditorMode/Source/ModelingToolsEditorMode/Private/ModelingToolsManagerActions.cpp
--- a/Engine/Plugins/Editor/ModelingToolsEditorMode/Source/ModelingToolsEditorMode/Private/ModelingToolsManagerActions.cpp
+++ b/Engine/Plugins/Editor/ModelingToolsEditorMode/Source/ModelingToolsEditorMode/Private/ModelingToolsManagerActions.cpp
@@ -393,19 +393,19 @@ void FModelingToolsManagerCommands::RegisterCommands()
 
 	UI_COMMAND(AddToFavorites, "Add to Favorites", "Add to Favorites", EUserInterfaceActionType::Button, FInputChord());
 	UI_COMMAND(RemoveFromFavorites, "Remove from Favorites", "Remove from Favorites", EUserInterfaceActionType::Button, FInputChord());
-	UI_COMMAND(LoadFavoritesTools, "Favorites", "Favorites", EUserInterfaceActionType::RadioButton, FInputChord());
-	UI_COMMAND(LoadSelectionTools, "Selection", "Edit Mesh Selections", EUserInterfaceActionType::RadioButton, FInputChord());
+	UI_COMMAND(LoadFavoritesTools, "Favs", "Favorites", EUserInterfaceActionType::RadioButton, FInputChord());
+	UI_COMMAND(LoadSelectionTools, "Select", "Edit Mesh Selections", EUserInterfaceActionType::RadioButton, FInputChord());
 	UI_COMMAND(LoadShapesTools, "Shapes", "Shapes", EUserInterfaceActionType::RadioButton, FInputChord());
 	UI_COMMAND(LoadCreateTools, "Create", "Create New Shapes", EUserInterfaceActionType::RadioButton, FInputChord());
 	UI_COMMAND(LoadPolyTools, "Model", "Shape Modeling", EUserInterfaceActionType::RadioButton, FInputChord());
 	UI_COMMAND(LoadTriTools, "Process", "Mesh Processing", EUserInterfaceActionType::RadioButton, FInputChord());
 	UI_COMMAND(LoadDeformTools, "Deform", "Deformations", EUserInterfaceActionType::RadioButton, FInputChord());
-	UI_COMMAND(LoadTransformTools, "Transform", "Transforms & Conversion", EUserInterfaceActionType::RadioButton, FInputChord());
-	UI_COMMAND(LoadMeshOpsTools, "MeshOps", "Mesh Processing", EUserInterfaceActionType::RadioButton, FInputChord());
-	UI_COMMAND(LoadVoxOpsTools, "VoxOps", "Voxel Processing", EUserInterfaceActionType::RadioButton, FInputChord());
-	UI_COMMAND(LoadAttributesTools, "Attributes", "Mesh Attributes", EUserInterfaceActionType::RadioButton, FInputChord());
+	UI_COMMAND(LoadTransformTools, "XForm", "Transforms & Conversion", EUserInterfaceActionType::RadioButton, FInputChord());
+	UI_COMMAND(LoadMeshOpsTools, "Mesh", "Mesh Processing", EUserInterfaceActionType::RadioButton, FInputChord());
+	UI_COMMAND(LoadVoxOpsTools, "Voxel", "Voxel Processing", EUserInterfaceActionType::RadioButton, FInputChord());
+	UI_COMMAND(LoadAttributesTools, "Attribs", "Mesh Attributes", EUserInterfaceActionType::RadioButton, FInputChord());
 	UI_COMMAND(LoadUVsTools, "UVs", "Create & Edit UVs", EUserInterfaceActionType::RadioButton, FInputChord());
-	UI_COMMAND(LoadBakingTools, "Baking", "Bake Textures & Colors", EUserInterfaceActionType::RadioButton, FInputChord());
+	UI_COMMAND(LoadBakingTools, "Bake", "Bake Textures & Colors", EUserInterfaceActionType::RadioButton, FInputChord());
 	UI_COMMAND(LoadVolumeTools, "Volumes", "Volumes", EUserInterfaceActionType::RadioButton, FInputChord());
 	UI_COMMAND(LoadLodsTools, "Misc", "Additional Utility Tools", EUserInterfaceActionType::RadioButton, FInputChord());
 
diff --git a/Engine/Plugins/Editor/ModelingToolsEditorMode/Source/ModelingToolsEditorMode/Public/ModelingToolsActions.h b/Engine/Plugins/Editor/ModelingToolsEditorMode/Source/ModelingToolsEditorMode/Public/ModelingToolsActions.h
--- a/Engine/Plugins/Editor/ModelingToolsEditorMode/Source/ModelingToolsEditorMode/Public/ModelingToolsActions.h
+++ b/Engine/Plugins/Editor/ModelingToolsEditorMode/Source/ModelingToolsEditorMode/Public/ModelingToolsActions.h
@@ -8,7 +8,8 @@
 
 enum class EModelingModeActionCommands
 {
-	FocusViewToCursor
+	FocusViewToCursor,
+	ToggleSelectionLockState
 };
 
 
@@ -20,6 +21,7 @@ public:
 	virtual void RegisterCommands() override;
 
 	TSharedPtr<FUICommandInfo> FocusViewCommand;
+	TSharedPtr<FUICommandInfo> ToggleSelectionLockStateCommand;
 
 	static void RegisterCommandBindings(TSharedPtr<FUICommandList> UICommandList, TFunction<void(EModelingModeActionCommands)> OnCommandExecuted);
 	static void UnRegisterCommandBindings(TSharedPtr<FUICommandList> UICommandList);
diff --git a/Engine/Plugins/Editor/ModelingToolsEditorMode/Source/ModelingToolsEditorMode/Public/ModelingToolsEditorMode.h b/Engine/Plugins/Editor/ModelingToolsEditorMode/Source/ModelingToolsEditorMode/Public/ModelingToolsEditorMode.h
--- a/Engine/Plugins/Editor/ModelingToolsEditorMode/Source/ModelingToolsEditorMode/Public/ModelingToolsEditorMode.h
+++ b/Engine/Plugins/Editor/ModelingToolsEditorMode/Source/ModelingToolsEditorMode/Public/ModelingToolsEditorMode.h
@@ -92,6 +92,8 @@ public:
 	UPROPERTY()
 	bool bEnableStaticMeshElementSelection = true;
 
+	bool GetMeshElementSelectionSystemEnabled() const;
+	void NotifySelectionSystemEnabledStateModified();
 
 protected:
 	virtual void BindCommands() override;
@@ -129,6 +131,9 @@ protected:
 	bool GetGeometrySelectionChangesAllowed() const;
 	bool TestForEditorGizmoHit(const FInputDeviceRay&) const;
 
+
+	bool bSelectionSystemEnabled = false;
+
 	void UpdateSelectionManagerOnEditorSelectionChange(bool bEnteringMode = false);
 
 	void OnToolsContextRender(IToolsContextRenderAPI* RenderAPI);
@@ -154,6 +159,9 @@ protected:
 	// analytics tracking
 	static FDateTime LastModeStartTimestamp;
 	static FDateTime LastToolStartTimestamp;
+
+	// tracking of unlocked stuff
+	static FDelegateHandle GlobalModelingWorldTeardownEventHandle;
 };
 
 #if UE_ENABLE_INCLUDE_ORDER_DEPRECATED_IN_5_2
diff --git a/Engine/Plugins/Editor/ModelingToolsEditorMode/Source/ModelingToolsEditorMode/Public/ModelingToolsEditorModeToolkit.h b/Engine/Plugins/Editor/ModelingToolsEditorMode/Source/ModelingToolsEditorMode/Public/ModelingToolsEditorModeToolkit.h
--- a/Engine/Plugins/Editor/ModelingToolsEditorMode/Source/ModelingToolsEditorMode/Public/ModelingToolsEditorModeToolkit.h
+++ b/Engine/Plugins/Editor/ModelingToolsEditorMode/Source/ModelingToolsEditorMode/Public/ModelingToolsEditorModeToolkit.h
@@ -91,13 +91,19 @@ public:
 	// get cached UGeometrySelectionManager pointer for the current Mode
 	UGeometrySelectionManager* GetMeshSelectionManager();
 
+	// @return true if we are currently in a tool, used to enable/disable UI things
+	bool IsInActiveTool() const { return bInActiveTool; }
+
 	// This is exposed only for the convenience of being able to create the numerical UI submenu
 	// in a non-member function in ModelingModeToolkit_Toolbars.cpp
 	TSharedPtr<STransformGizmoNumericalUIOverlay> GetGizmoNumericalUIOverlayWidget() { return GizmoNumericalUIOverlayWidget; }
 
+	void NotifySelectionSystemEnabledStateModified();
+
 private:
 	const static TArray<FName> PaletteNames_Standard;
 
+	bool bInActiveTool = false;
 	FText ActiveToolName;
 	FText ActiveToolMessage;
 	FStatusBarMessageHandle ActiveToolMessageHandle;
diff --git a/Engine/Plugins/Runtime/MeshModelingToolset/Source/ModelingComponents/Private/Selection/GeometrySelectionManager.cpp b/Engine/Plugins/Runtime/MeshModelingToolset/Source/ModelingComponents/Private/Selection/GeometrySelectionManager.cpp
--- a/Engine/Plugins/Runtime/MeshModelingToolset/Source/ModelingComponents/Private/Selection/GeometrySelectionManager.cpp
+++ b/Engine/Plugins/Runtime/MeshModelingToolset/Source/ModelingComponents/Private/Selection/GeometrySelectionManager.cpp
@@ -440,15 +440,16 @@ void UGeometrySelectionManager::SetCurrentTargetsLockState(bool bLocked)
 
 			if (!bInTransaction)
 			{
-				GetTransactionsAPI()->BeginUndoTransaction((bLocked) ?
-					LOCTEXT("Lock Target", "Lock Target") : LOCTEXT("Unlock Target", "Unlock Target"));
+				GetTransactionsAPI()->BeginUndoTransaction(
+					(bLocked) ? LOCTEXT("Lock Target", "Lock Target") : LOCTEXT("Unlock Target", "Unlock Target"));
 				bInTransaction = true;
 			}
 
 			TUniquePtr<FGeometrySelectionManager_TargetLockStateChange> Change = MakeUnique<FGeometrySelectionManager_TargetLockStateChange>();
 			Change->TargetIdentifier = Target->TargetIdentifier;
 			Change->bToState = bLocked;
-			GetTransactionsAPI()->AppendChange(this, MoveTemp(Change), LOCTEXT("Lock Target", "Lock Target"));
+			GetTransactionsAPI()->AppendChange(this, MoveTemp(Change), 
+				(bLocked) ? LOCTEXT("Lock Target", "Lock Target") : LOCTEXT("Unlock Target", "Unlock Target"));
 		}
 	}
 
diff --git a/Engine/Plugins/Runtime/MeshModelingToolset/Source/ModelingComponentsEditorOnly/Private/Selection/StaticMeshSelector.cpp b/Engine/Plugins/Runtime/MeshModelingToolset/Source/ModelingComponentsEditorOnly/Private/Selection/StaticMeshSelector.cpp
--- a/Engine/Plugins/Runtime/MeshModelingToolset/Source/ModelingComponentsEditorOnly/Private/Selection/StaticMeshSelector.cpp
+++ b/Engine/Plugins/Runtime/MeshModelingToolset/Source/ModelingComponentsEditorOnly/Private/Selection/StaticMeshSelector.cpp
@@ -24,15 +24,38 @@ using namespace UE::Geometry;
 
 #define LOCTEXT_NAMESPACE "FStaticMeshSelector"
 
+static TAutoConsoleVariable<bool> CVarEnableModelingSelectionStaticMeshLocking(
+	TEXT("modeling.Selection.EnableStaticMeshLocking"),
+	true,
+	TEXT("Control whether Selection Locking is enabled by default for Static Meshes"));
+
 namespace UEGlobal
 {
+	/**
+	 * Only Unlocked Mesh Assets can be edited. This is currently implemented as a global set based on UObject pointers.
+	 * Obviously not an ideal solution, however this is really only used to implement a UI gate, if we get a stale pointer
+	 * here that happens to still point to a valid UStaticMesh, it just means that asset will be unlocked for mesh selection
+	 * even if the user did not explicitly unlock it. This is not a disaster.
+	 */
 	TSet<UStaticMesh*> UnlockedStaticMeshes;
 }
 
+bool FStaticMeshSelector::IsLockable() const 
+{ 
+	return CVarEnableModelingSelectionStaticMeshLocking.GetValueOnGameThread();
+}
+
 
 bool FStaticMeshSelector::IsLocked() const
 {
-	return StaticMesh != nullptr && (UEGlobal::UnlockedStaticMeshes.Contains(StaticMesh) == false);
+	if (CVarEnableModelingSelectionStaticMeshLocking.GetValueOnGameThread())
+	{
+		return StaticMesh != nullptr && (UEGlobal::UnlockedStaticMeshes.Contains(StaticMesh) == false);
+	}
+	else
+	{
+		return (StaticMesh == nullptr);
+	}
 }
 
 void FStaticMeshSelector::SetLockedState(bool bLocked)
@@ -59,6 +82,20 @@ void FStaticMeshSelector::SetLockedState(bool bLocked)
 }
 
 
+void FStaticMeshSelector::SetAssetUnlockedOnCreation(UStaticMesh* StaticMesh)
+{
+	if (StaticMesh != nullptr)
+	{
+		UEGlobal::UnlockedStaticMeshes.Add(StaticMesh);
+	}
+}
+
+void FStaticMeshSelector::ResetUnlockedStaticMeshAssets()
+{
+	UEGlobal::UnlockedStaticMeshes.Reset();
+}
+
+
 
 bool FStaticMeshSelector::Initialize(
 	FGeometryIdentifier SourceGeometryIdentifierIn)
diff --git a/Engine/Plugins/Runtime/MeshModelingToolset/Source/ModelingComponentsEditorOnly/Private/Selection/VolumeSelector.cpp b/Engine/Plugins/Runtime/MeshModelingToolset/Source/ModelingComponentsEditorOnly/Private/Selection/VolumeSelector.cpp
--- a/Engine/Plugins/Runtime/MeshModelingToolset/Source/ModelingComponentsEditorOnly/Private/Selection/VolumeSelector.cpp
+++ b/Engine/Plugins/Runtime/MeshModelingToolset/Source/ModelingComponentsEditorOnly/Private/Selection/VolumeSelector.cpp
@@ -24,19 +24,43 @@ using namespace UE::Geometry;
 
 #define LOCTEXT_NAMESPACE "FVolumeSelector"
 
+static TAutoConsoleVariable<bool> CVarEnableModelingSelectionVolumeLocking(
+	TEXT("modeling.Selection.EnableVolumeLocking"),
+	true,
+	TEXT("Control whether Selection Locking is enabled by default for Volumes"));
 
 
 namespace UEGlobal
 {
+	/**
+	 * Only Unlocked Brush Components can be edited. This is currently implemented as a global set based on UObject pointers.
+	 * Obviously not an ideal solution, however this is really only used to implement a UI gate, if we get a stale pointer
+	 * here that happens to still point to a valid UBrushComponent, it just means that Volume will be unlocked for mesh selection
+	 * even if the user did not explicitly unlock it. This is not a disaster.
+	 */
 	TSet<UBrushComponent*> UnlockedBrushComponents;
 }
 
 
+bool FVolumeSelector::IsLockable() const
+{
+	return CVarEnableModelingSelectionVolumeLocking.GetValueOnGameThread();
+}
+
+
 bool FVolumeSelector::IsLocked() const
 {
-	return BrushComponent != nullptr && (UEGlobal::UnlockedBrushComponents.Contains(BrushComponent) == false);
+	if (CVarEnableModelingSelectionVolumeLocking.GetValueOnGameThread())
+	{
+		return BrushComponent != nullptr && (UEGlobal::UnlockedBrushComponents.Contains(BrushComponent) == false);
+	}
+	else
+	{
+		return (BrushComponent == nullptr);
+	}
 }
 
+
 void FVolumeSelector::SetLockedState(bool bLocked)
 {
 	if (BrushComponent != nullptr)
@@ -56,6 +80,18 @@ void FVolumeSelector::SetLockedState(bool bLocked)
 	}
 }
 
+void FVolumeSelector::SetComponentUnlockedOnCreation(UBrushComponent* BrushComponent)
+{
+	if (BrushComponent != nullptr)
+	{
+		UEGlobal::UnlockedBrushComponents.Add(BrushComponent);
+	}
+}
+
+void FVolumeSelector::ResetUnlockedBrushComponents()
+{
+	UEGlobal::UnlockedBrushComponents.Reset();
+}
 
 
 bool FVolumeSelector::Initialize(
diff --git a/Engine/Plugins/Runtime/MeshModelingToolset/Source/ModelingComponentsEditorOnly/Public/Selection/StaticMeshSelector.h b/Engine/Plugins/Runtime/MeshModelingToolset/Source/ModelingComponentsEditorOnly/Public/Selection/StaticMeshSelector.h
--- a/Engine/Plugins/Runtime/MeshModelingToolset/Source/ModelingComponentsEditorOnly/Public/Selection/StaticMeshSelector.h
+++ b/Engine/Plugins/Runtime/MeshModelingToolset/Source/ModelingComponentsEditorOnly/Public/Selection/StaticMeshSelector.h
@@ -33,7 +33,7 @@ public:
 	//virtual bool Restore();
 
 
-	virtual bool IsLockable() const override { return true; }
+	virtual bool IsLockable() const override;
 	virtual bool IsLocked() const override;
 	virtual void SetLockedState(bool bLocked) override;
 
@@ -52,6 +52,11 @@ protected:
 
 	TSharedPtr<FBasicDynamicMeshSelectionTransformer> ActiveTransformer;
 	void CommitMeshTransform();
+
+
+public:
+	static void SetAssetUnlockedOnCreation(UStaticMesh* StaticMesh);
+	static void ResetUnlockedStaticMeshAssets();
 };
 
 
diff --git a/Engine/Plugins/Runtime/MeshModelingToolset/Source/ModelingComponentsEditorOnly/Public/Selection/VolumeSelector.h b/Engine/Plugins/Runtime/MeshModelingToolset/Source/ModelingComponentsEditorOnly/Public/Selection/VolumeSelector.h
--- a/Engine/Plugins/Runtime/MeshModelingToolset/Source/ModelingComponentsEditorOnly/Public/Selection/VolumeSelector.h
+++ b/Engine/Plugins/Runtime/MeshModelingToolset/Source/ModelingComponentsEditorOnly/Public/Selection/VolumeSelector.h
@@ -32,7 +32,7 @@ public:
 	//virtual bool Sleep();
 	//virtual bool Restore();
 
-	virtual bool IsLockable() const override { return true; }
+	virtual bool IsLockable() const override;
 	virtual bool IsLocked() const override;
 	virtual void SetLockedState(bool bLocked) override;
 
@@ -44,6 +44,10 @@ public:
 	virtual void PostUndo(bool bSuccess) override;
 	virtual void PostRedo(bool bSuccess) override;
 
+public:
+	static void SetComponentUnlockedOnCreation(UBrushComponent* Component);
+	static void ResetUnlockedBrushComponents();
+
 protected:
 	AVolume* ParentVolume = nullptr;
 	UBrushComponent* BrushComponent = nullptr;
