fix bad merge

diff --git a/Engine/Plugins/Editor/ModelingToolsEditorMode/Source/ModelingToolsEditorMode/Private/ModelingToolsEditorModeToolkit.cpp b/Engine/Plugins/Editor/ModelingToolsEditorMode/Source/ModelingToolsEditorMode/Private/ModelingToolsEditorModeToolkit.cpp
--- a/Engine/Plugins/Editor/ModelingToolsEditorMode/Source/ModelingToolsEditorMode/Private/ModelingToolsEditorModeToolkit.cpp
+++ b/Engine/Plugins/Editor/ModelingToolsEditorMode/Source/ModelingToolsEditorMode/Private/ModelingToolsEditorModeToolkit.cpp
@@ -589,7 +589,6 @@ void FModelingToolsEditorModeToolkit::Init(const TSharedPtr<IToolkitHost>& InitT
 
 	MakeSelectionPaletteOverlayWidget();
 	GetToolkitHost()->AddViewportOverlayWidget(SelectionPaletteOverlayWidget.ToSharedRef());
-}
 
 	CurrentPreset = FSoftObjectPath(); // Default to the default collection by leaving this null.
 }
