[Backout] - CL25182400 for CIS issue
#rnx
Original CL Desc
-----------------------------------------------------------------
MeshModelingTools: Add Revolve Spline tool.

diff --git a/Engine/Plugins/Editor/ModelingToolsEditorMode/Source/ModelingToolsEditorMode/Private/ModelingToolsEditorMode.cpp b/Engine/Plugins/Editor/ModelingToolsEditorMode/Source/ModelingToolsEditorMode/Private/ModelingToolsEditorMode.cpp
--- a/Engine/Plugins/Editor/ModelingToolsEditorMode/Source/ModelingToolsEditorMode/Private/ModelingToolsEditorMode.cpp
+++ b/Engine/Plugins/Editor/ModelingToolsEditorMode/Source/ModelingToolsEditorMode/Private/ModelingToolsEditorMode.cpp
@@ -44,7 +44,6 @@
 #include "DrawPolygonTool.h"
 #include "DrawPolyPathTool.h"
 #include "DrawAndRevolveTool.h"
-#include "RevolveSplineTool.h"
 #include "ShapeSprayTool.h"
 #include "VoxelSolidifyMeshesTool.h"
 #include "VoxelBlendMeshesTool.h"
@@ -573,9 +572,6 @@ void UModelingToolsEditorMode::Enter()
 	auto DrawAndRevolveToolBuilder = NewObject<UDrawAndRevolveToolBuilder>();
 	RegisterTool(ToolManagerCommands.BeginDrawAndRevolveTool, TEXT("BeginDrawAndRevolveTool"), DrawAndRevolveToolBuilder);
 
-	auto RevolveSplineToolBuilder = NewObject<URevolveSplineToolBuilder>();
-	RegisterTool(ToolManagerCommands.BeginRevolveSplineTool, TEXT("BeginRevolveSplineTool"), RevolveSplineToolBuilder);
-
 	auto ShapeSprayToolBuilder = NewObject<UShapeSprayToolBuilder>();
 	RegisterTool(ToolManagerCommands.BeginShapeSprayTool, TEXT("BeginShapeSprayTool"), ShapeSprayToolBuilder);
 
diff --git a/Engine/Plugins/Editor/ModelingToolsEditorMode/Source/ModelingToolsEditorMode/Private/ModelingToolsEditorModeToolkit.cpp b/Engine/Plugins/Editor/ModelingToolsEditorMode/Source/ModelingToolsEditorMode/Private/ModelingToolsEditorModeToolkit.cpp
--- a/Engine/Plugins/Editor/ModelingToolsEditorMode/Source/ModelingToolsEditorMode/Private/ModelingToolsEditorModeToolkit.cpp
+++ b/Engine/Plugins/Editor/ModelingToolsEditorMode/Source/ModelingToolsEditorMode/Private/ModelingToolsEditorModeToolkit.cpp
@@ -257,7 +257,6 @@ void FModelingToolsEditorModeToolkit::RegisterPalettes()
 		Commands.BeginDrawPolygonTool,
 		Commands.BeginDrawPolyPathTool,
 		Commands.BeginDrawAndRevolveTool,
-		Commands.BeginRevolveSplineTool,
 		Commands.BeginTriangulateSplinesTool
 	});
 	ToolkitBuilder->AddPalette(
@@ -1623,7 +1622,6 @@ void FModelingToolsEditorModeToolkit::BuildToolPalette(FName PaletteIndex, class
 		ToolbarBuilder.AddToolBarButton(Commands.BeginDrawPolygonTool);
 		ToolbarBuilder.AddToolBarButton(Commands.BeginDrawPolyPathTool);
 		ToolbarBuilder.AddToolBarButton(Commands.BeginDrawAndRevolveTool);
-		ToolbarBuilder.AddToolBarButton(Commands.BeginRevolveSplineTool);
 		ToolbarBuilder.AddToolBarButton(Commands.BeginRevolveBoundaryTool);
 		ToolbarBuilder.AddToolBarButton(Commands.BeginCombineMeshesTool);
 		ToolbarBuilder.AddToolBarButton(Commands.BeginDuplicateMeshesTool);
diff --git a/Engine/Plugins/Editor/ModelingToolsEditorMode/Source/ModelingToolsEditorMode/Private/ModelingToolsManagerActions.cpp b/Engine/Plugins/Editor/ModelingToolsEditorMode/Source/ModelingToolsEditorMode/Private/ModelingToolsManagerActions.cpp
--- a/Engine/Plugins/Editor/ModelingToolsEditorMode/Source/ModelingToolsEditorMode/Private/ModelingToolsManagerActions.cpp
+++ b/Engine/Plugins/Editor/ModelingToolsEditorMode/Source/ModelingToolsEditorMode/Private/ModelingToolsManagerActions.cpp
@@ -79,7 +79,6 @@ void FModelingToolsManagerCommands::RegisterCommands()
 		REGISTER_MODELING_TOOL_COMMAND(BeginDrawPolygonTool, "PolyExt", "Draw and extrude 2D Polygons to create new objects");
 		REGISTER_MODELING_TOOL_COMMAND(BeginDrawPolyPathTool, "PathExt", "Draw and extrude 2D Paths to create new objects");
 		REGISTER_MODELING_TOOL_COMMAND(BeginDrawAndRevolveTool, "PathRev", "Draw and revolve 2D Paths to create new objects");
-		REGISTER_MODELING_TOOL_COMMAND(BeginRevolveSplineTool, "SplnRev", "Revolve splines to create new objects");
 		REGISTER_MODELING_TOOL_COMMAND(BeginRevolveBoundaryTool, "BdryRev", "Revolve Mesh boundary loops to create new objects");
 		REGISTER_MODELING_TOOL_COMMAND(BeginCombineMeshesTool, "Merge", "Merge multiple Meshes to create new objects");
 		REGISTER_MODELING_TOOL_COMMAND(BeginDuplicateMeshesTool, "Dupe", "Duplicate single Meshes to create new objects");
@@ -90,7 +89,6 @@ void FModelingToolsManagerCommands::RegisterCommands()
 		REGISTER_MODELING_TOOL_COMMAND(BeginDrawPolygonTool, "Extrude Polygon", "Draw and extrude 2D Polygons to create new objects");
 		REGISTER_MODELING_TOOL_COMMAND(BeginDrawPolyPathTool, "Extrude Path", "Draw and extrude 2D Paths to create new objects");
 		REGISTER_MODELING_TOOL_COMMAND(BeginDrawAndRevolveTool, "Revolve Path", "Draw and revolve 2D Paths to create new objects");
-		REGISTER_MODELING_TOOL_COMMAND(BeginRevolveSplineTool, "Revolve Spline", "Revolve splines to create new objects");
 		REGISTER_MODELING_TOOL_COMMAND(BeginRevolveBoundaryTool, "Revolve Boundary", "Revolve Mesh boundary loops to create new objects");
 		REGISTER_MODELING_TOOL_COMMAND(BeginCombineMeshesTool, "Merge", "Merge multiple Meshes to create new objects");
 		REGISTER_MODELING_TOOL_COMMAND(BeginDuplicateMeshesTool, "Duplicate", "Duplicate single Meshes to create new objects");
diff --git a/Engine/Plugins/Editor/ModelingToolsEditorMode/Source/ModelingToolsEditorMode/Public/ModelingToolsManagerActions.h b/Engine/Plugins/Editor/ModelingToolsEditorMode/Source/ModelingToolsEditorMode/Public/ModelingToolsManagerActions.h
--- a/Engine/Plugins/Editor/ModelingToolsEditorMode/Source/ModelingToolsEditorMode/Public/ModelingToolsManagerActions.h
+++ b/Engine/Plugins/Editor/ModelingToolsEditorMode/Source/ModelingToolsEditorMode/Public/ModelingToolsManagerActions.h
@@ -46,7 +46,6 @@ public:
 	TSharedPtr<FUICommandInfo> BeginDrawPolygonTool;
 	TSharedPtr<FUICommandInfo> BeginDrawPolyPathTool;
 	TSharedPtr<FUICommandInfo> BeginDrawAndRevolveTool;
-	TSharedPtr<FUICommandInfo> BeginRevolveSplineTool;
 	TSharedPtr<FUICommandInfo> BeginShapeSprayTool;
 	TSharedPtr<FUICommandInfo> BeginTriangulateSplinesTool;
 
diff --git a/Engine/Plugins/Experimental/MeshModelingToolsetExp/Source/MeshModelingToolsExp/Private/RevolveSplineTool.cpp b/Engine/Plugins/Experimental/MeshModelingToolsetExp/Source/MeshModelingToolsExp/Private/RevolveSplineTool.cpp
deleted file mode 100644
diff --git a/Engine/Plugins/Experimental/MeshModelingToolsetExp/Source/MeshModelingToolsExp/Public/RevolveSplineTool.h b/Engine/Plugins/Experimental/MeshModelingToolsetExp/Source/MeshModelingToolsExp/Public/RevolveSplineTool.h
deleted file mode 100644
diff --git a/Engine/Plugins/Runtime/MeshModelingToolset/Source/ModelingOperators/Private/CompositionOps/CurveSweepOp.cpp b/Engine/Plugins/Runtime/MeshModelingToolset/Source/ModelingOperators/Private/CompositionOps/CurveSweepOp.cpp
--- a/Engine/Plugins/Runtime/MeshModelingToolset/Source/ModelingOperators/Private/CompositionOps/CurveSweepOp.cpp
+++ b/Engine/Plugins/Runtime/MeshModelingToolset/Source/ModelingOperators/Private/CompositionOps/CurveSweepOp.cpp
@@ -112,7 +112,6 @@ void FCurveSweepOp::CalculateResult(FProgressCancel* Progress)
 					ConstrainedDelaunayTriangulate<double>,
 					PlaneOrigin,
 					PlaneNormal);
-				break;
 			default:
 				checkSlow(false);
 			}
