MeshModelingTools: Fix gizmo location being set incorrectly after various tools that edit the location without editing selection.

diff --git a/Engine/Plugins/Editor/ModelingToolsEditorMode/Source/ModelingToolsEditorMode/Private/ModelingToolsEditorMode.cpp b/Engine/Plugins/Editor/ModelingToolsEditorMode/Source/ModelingToolsEditorMode/Private/ModelingToolsEditorMode.cpp
--- a/Engine/Plugins/Editor/ModelingToolsEditorMode/Source/ModelingToolsEditorMode/Private/ModelingToolsEditorMode.cpp
+++ b/Engine/Plugins/Editor/ModelingToolsEditorMode/Source/ModelingToolsEditorMode/Private/ModelingToolsEditorMode.cpp
@@ -3,6 +3,7 @@
 #include "ModelingToolsEditorMode.h"
 #include "Components/StaticMeshComponent.h"
 #include "IGeometryProcessingInterfacesModule.h"
+#include "Editor/UnrealEdEngine.h"
 #include "EdModeInteractiveToolsContext.h"
 #include "GeometryProcessingInterfaces/IUVEditorModularFeature.h"
 #include "IAnalyticsProviderET.h"
@@ -22,6 +23,7 @@
 #include "BaseGizmos/TransformGizmoUtil.h"
 #include "Snapping/ModelingSceneSnappingManager.h"
 #include "Scene/LevelObjectsObserver.h"
+#include "UnrealEdGlobals.h" // GUnrealEd
 
 #include "Features/IModularFeatures.h"
 #include "ModelingModeToolExtensions.h"
@@ -1468,6 +1470,10 @@ void UModelingToolsEditorMode::OnToolEnded(UInteractiveToolManager* Manager, UIn
 
 	FModelingToolActionCommands::UpdateToolCommandBinding(Tool, Toolkit->GetToolkitCommands(), true);
 	
+	// We may require a gizmo location update despite not having changed the selection (transform tool,
+	// edit pivot, etc).
+	GUnrealEd->UpdatePivotLocationForSelection();
+
 	if( FEngineAnalytics::IsAvailable() )
 	{
 		FEngineAnalytics::GetProvider().RecordEvent(TEXT("Editor.Usage.MeshModelingMode.ToolEnded"),
