ModelingMode: improve handling around tool palette visibility. Listen to new delegate in ToolkitBuilder to know when palette changes, and re-show hidden palette buttons on tab change. Add new Editor setting to disable hiding of palette buttons. speculatively Harden some nullptr checking around ToolManager

diff --git a/Engine/Plugins/Editor/ModelingToolsEditorMode/Source/ModelingToolsEditorMode/Private/ModelingModeToolkit_Toolbars.cpp b/Engine/Plugins/Editor/ModelingToolsEditorMode/Source/ModelingToolsEditorMode/Private/ModelingModeToolkit_Toolbars.cpp
--- a/Engine/Plugins/Editor/ModelingToolsEditorMode/Source/ModelingToolsEditorMode/Private/ModelingModeToolkit_Toolbars.cpp
+++ b/Engine/Plugins/Editor/ModelingToolsEditorMode/Source/ModelingToolsEditorMode/Private/ModelingModeToolkit_Toolbars.cpp
@@ -726,8 +726,6 @@ void FModelingToolsEditorModeToolkit::MakeSelectionPaletteOverlayWidget()
 		FSlateIcon(FModelingToolsEditorModeStyle::GetStyleSetName(), "SelectionToolBarIcons.UnlockedTarget"));
 
 
-
-
 	TSharedRef<SWidget> Toolbar = ToolbarBuilder.MakeWidget();
 
 	SAssignNew(SelectionPaletteOverlayWidget, SVerticalBox)
@@ -738,24 +736,23 @@ void FModelingToolsEditorModeToolkit::MakeSelectionPaletteOverlayWidget()
 		[
 			SNew(SBorder)
 			.BorderImage(FAppStyle::Get().GetBrush("EditorViewport.OverlayBrush"))
-		.Padding(FMargin(3.0f, 6.0f, 3.f, 6.f))
-		[
-			SNew(SVerticalBox)
-
-			+ SVerticalBox::Slot()
-		.AutoHeight()
-		.Padding(FMargin(0.0, 0.f, 0.f, 0.f))
-		[
-			Toolbar
-		]
-		]
+			.Padding(FMargin(3.0f, 6.0f, 3.f, 6.f))
+			[
+				SNew(SVerticalBox)
+				+ SVerticalBox::Slot()
+				.AutoHeight()
+				.Padding(FMargin(0.0, 0.f, 0.f, 0.f))
+				[
+					Toolbar
+				]
+			]
 		];
 
 	SelectionPaletteOverlayWidget->SetVisibility( TAttribute<EVisibility>::CreateLambda([this]()
 	{
 		if ( UModelingToolsEditorMode* ModelingMode = Cast<UModelingToolsEditorMode>(GetScriptableEditorMode()) )
 		{
-			return ModelingMode->GetMeshElementSelectionSystemEnabled() ? EVisibility::Visible : EVisibility::Collapsed;
+			return ModelingMode->GetMeshElementSelectionSystemEnabled() ? EVisibility::SelfHitTestInvisible : EVisibility::Collapsed;
 		}
 		return EVisibility::Collapsed;
 	}) );
diff --git a/Engine/Plugins/Editor/ModelingToolsEditorMode/Source/ModelingToolsEditorMode/Private/ModelingToolsEditorMode.cpp b/Engine/Plugins/Editor/ModelingToolsEditorMode/Source/ModelingToolsEditorMode/Private/ModelingToolsEditorMode.cpp
--- a/Engine/Plugins/Editor/ModelingToolsEditorMode/Source/ModelingToolsEditorMode/Private/ModelingToolsEditorMode.cpp
+++ b/Engine/Plugins/Editor/ModelingToolsEditorMode/Source/ModelingToolsEditorMode/Private/ModelingToolsEditorMode.cpp
@@ -834,18 +834,22 @@ void UModelingToolsEditorMode::Enter()
 	{
 		Toolkit->GetToolkitCommands()->MapAction(UICommand,
 			FExecuteAction::CreateLambda([this, TopoMode, ElementMode]() {
-				GetToolManager()->GetContextTransactionsAPI()->BeginUndoTransaction(LOCTEXT("ChangeSelectionMode", "Selection Mode"));
-				GetSelectionManager()->SetMeshTopologyMode(TopoMode); 
-				GetSelectionManager()->SetSelectionElementType(ElementMode);
-				GetToolManager()->GetContextTransactionsAPI()->EndUndoTransaction();
-
-				UModelingToolsModeCustomizationSettings* ModelingEditorSettings = GetMutableDefault<UModelingToolsModeCustomizationSettings>();
-				ModelingEditorSettings->LastMeshSelectionTopologyMode = static_cast<int>(TopoMode);
-				ModelingEditorSettings->LastMeshSelectionElementType = static_cast<int>(ElementMode);
-				ModelingEditorSettings->SaveConfig();
+				if ( GetToolManager() && GetToolManager()->GetContextTransactionsAPI() && GetSelectionManager() )
+				{
+					GetToolManager()->GetContextTransactionsAPI()->BeginUndoTransaction(LOCTEXT("ChangeSelectionMode", "Selection Mode"));
+					GetSelectionManager()->SetMeshTopologyMode(TopoMode); 
+					GetSelectionManager()->SetSelectionElementType(ElementMode);
+					GetToolManager()->GetContextTransactionsAPI()->EndUndoTransaction();
+					if (UModelingToolsModeCustomizationSettings* ModelingEditorSettings = GetMutableDefault<UModelingToolsModeCustomizationSettings>())
+					{
+						ModelingEditorSettings->LastMeshSelectionTopologyMode = static_cast<int>(TopoMode);
+						ModelingEditorSettings->LastMeshSelectionElementType = static_cast<int>(ElementMode);
+						ModelingEditorSettings->SaveConfig();
+					}
+				}
 			}),
-			FCanExecuteAction::CreateLambda([this]() { return GetToolManager()->HasAnyActiveTool() == false; }),
-			FIsActionChecked::CreateLambda([this, TopoMode, ElementMode]() { return GetSelectionManager()->GetMeshTopologyMode() == TopoMode && GetSelectionManager()->GetSelectionElementType() == ElementMode; }),
+			FCanExecuteAction::CreateLambda([this]() { return (GetToolManager() != nullptr && GetToolManager()->HasAnyActiveTool() == false) ? true : false; }),
+			FIsActionChecked::CreateLambda([this, TopoMode, ElementMode]() { return (GetSelectionManager() != nullptr && GetSelectionManager()->GetMeshTopologyMode() == TopoMode && GetSelectionManager()->GetSelectionElementType() == ElementMode); }),
 			EUIActionRepeatMode::RepeatDisabled);
 	};
 	if (GetSelectionManager() != nullptr)
diff --git a/Engine/Plugins/Editor/ModelingToolsEditorMode/Source/ModelingToolsEditorMode/Private/ModelingToolsEditorModeToolkit.cpp b/Engine/Plugins/Editor/ModelingToolsEditorMode/Source/ModelingToolsEditorMode/Private/ModelingToolsEditorModeToolkit.cpp
--- a/Engine/Plugins/Editor/ModelingToolsEditorMode/Source/ModelingToolsEditorMode/Private/ModelingToolsEditorModeToolkit.cpp
+++ b/Engine/Plugins/Editor/ModelingToolsEditorMode/Source/ModelingToolsEditorMode/Private/ModelingToolsEditorModeToolkit.cpp
@@ -440,14 +440,6 @@ void FModelingToolsEditorModeToolkit::RegisterPalettes()
 
 
 
-	//TArray<TSharedPtr<FUICommandInfo>> VolumePaletteItems({
-
-	//});
-	//ToolkitBuilder->AddPalette(
-	//	MakeShareable( new FToolPalette( Commands.LoadVolumeTools.ToSharedRef(), VolumePaletteItems)) );
-
-
-
 	TArray<TSharedPtr<FUICommandInfo>> LODPaletteItems({
 		Commands.BeginGenerateStaticMeshLODAssetTool,
 		Commands.BeginAddPivotActorTool,
@@ -464,6 +456,12 @@ void FModelingToolsEditorModeToolkit::RegisterPalettes()
 
 	ToolkitBuilder->SetActivePaletteOnLoad(Commands.LoadCreateTools.Get());
 	ToolkitBuilder->UpdateWidget();
+
+	// if selected palette changes, make sure we are showing the palette command buttons, which may be hidden by active Tool
+	ActivePaletteChangedHandle = ToolkitBuilder->OnActivePaletteChanged.AddLambda([this]()
+	{
+		ToolkitBuilder->SetActivePaletteCommandsVisibility(EVisibility::Visible);
+	});
 }
 
 
@@ -2048,7 +2046,13 @@ void FModelingToolsEditorModeToolkit::OnToolStarted(UInteractiveToolManager* Man
 	if (HasToolkitBuilder())
 	{
 		ToolkitBuilder->SetActiveToolDisplayName(ActiveToolName);
-		ToolkitBuilder->SetActivePaletteCommandsVisibility(EVisibility::Collapsed);
+		if (const UModelingToolsModeCustomizationSettings* Settings = GetDefault<UModelingToolsModeCustomizationSettings>())
+		{
+			if (Settings->bAlwaysShowToolButtons == false)
+			{
+				ToolkitBuilder->SetActivePaletteCommandsVisibility(EVisibility::Collapsed);
+			}
+		}
 	}
 
 	// try to update icon
@@ -2125,6 +2129,11 @@ void FModelingToolsEditorModeToolkit::OnActiveViewportChanged(TSharedPtr<IAssetV
 			{
 				GetToolkitHost()->RemoveViewportOverlayWidget(GizmoNumericalUIOverlayWidget.ToSharedRef(), OldViewport);
 			}
+
+			if (SelectionPaletteOverlayWidget.IsValid())
+			{
+				GetToolkitHost()->RemoveViewportOverlayWidget(SelectionPaletteOverlayWidget.ToSharedRef(), OldViewport);
+			}
 		}
 
 		// Add the hud to the new viewport
@@ -2134,6 +2143,11 @@ void FModelingToolsEditorModeToolkit::OnActiveViewportChanged(TSharedPtr<IAssetV
 		{
 			GetToolkitHost()->AddViewportOverlayWidget(GizmoNumericalUIOverlayWidget.ToSharedRef(), NewViewport);
 		}
+
+		if (SelectionPaletteOverlayWidget.IsValid())
+		{
+			GetToolkitHost()->AddViewportOverlayWidget(SelectionPaletteOverlayWidget.ToSharedRef(), NewViewport);
+		}
 	}
 }
 
diff --git a/Engine/Plugins/Editor/ModelingToolsEditorMode/Source/ModelingToolsEditorMode/Private/ModelingToolsManagerActions.cpp b/Engine/Plugins/Editor/ModelingToolsEditorMode/Source/ModelingToolsEditorMode/Private/ModelingToolsManagerActions.cpp
--- a/Engine/Plugins/Editor/ModelingToolsEditorMode/Source/ModelingToolsEditorMode/Private/ModelingToolsManagerActions.cpp
+++ b/Engine/Plugins/Editor/ModelingToolsEditorMode/Source/ModelingToolsEditorMode/Private/ModelingToolsManagerActions.cpp
@@ -397,7 +397,7 @@ void FModelingToolsManagerCommands::RegisterCommands()
 
 	UI_COMMAND(AddToFavorites, "Add to Favorites", "Add to Favorites", EUserInterfaceActionType::Button, FInputChord());
 	UI_COMMAND(RemoveFromFavorites, "Remove from Favorites", "Remove from Favorites", EUserInterfaceActionType::Button, FInputChord());
-	UI_COMMAND(LoadFavoritesTools, "Favs", "Favorites", EUserInterfaceActionType::RadioButton, FInputChord());
+	UI_COMMAND(LoadFavoritesTools, "Faves", "Favorites", EUserInterfaceActionType::RadioButton, FInputChord());
 	UI_COMMAND(LoadSelectionTools, "Select", "Edit Mesh Selections", EUserInterfaceActionType::RadioButton, FInputChord());
 	UI_COMMAND(LoadShapesTools, "Shapes", "Shapes", EUserInterfaceActionType::RadioButton, FInputChord());
 	UI_COMMAND(LoadCreateTools, "Create", "Create New Shapes", EUserInterfaceActionType::RadioButton, FInputChord());
@@ -426,5 +426,4 @@ void FModelingToolsManagerCommands::RegisterCommands()
 
 
 
-
 #undef LOCTEXT_NAMESPACE
diff --git a/Engine/Plugins/Editor/ModelingToolsEditorMode/Source/ModelingToolsEditorMode/Public/ModelingToolsEditorModeSettings.h b/Engine/Plugins/Editor/ModelingToolsEditorMode/Source/ModelingToolsEditorMode/Public/ModelingToolsEditorModeSettings.h
--- a/Engine/Plugins/Editor/ModelingToolsEditorMode/Source/ModelingToolsEditorMode/Public/ModelingToolsEditorModeSettings.h
+++ b/Engine/Plugins/Editor/ModelingToolsEditorMode/Source/ModelingToolsEditorMode/Public/ModelingToolsEditorModeSettings.h
@@ -345,6 +345,12 @@ public:
 	UPROPERTY(config, EditAnywhere, Category = "Modeling Mode|UI Customization")
 	bool bShowCategoryButtonLabels = true;
 	
+	/**
+	 * If true, Tool buttons will always be shown when in a Tool. By default they will be hidden.
+	 */
+	UPROPERTY(config, EditAnywhere, Category = "Modeling Mode|UI Customization")
+	bool bAlwaysShowToolButtons = false;
+
 public:
 
 	// saved-state for various mode settings that are configured via UI toggles/etc, and not exposed in settings dialog
diff --git a/Engine/Plugins/Editor/ModelingToolsEditorMode/Source/ModelingToolsEditorMode/Public/ModelingToolsEditorModeToolkit.h b/Engine/Plugins/Editor/ModelingToolsEditorMode/Source/ModelingToolsEditorMode/Public/ModelingToolsEditorModeToolkit.h
--- a/Engine/Plugins/Editor/ModelingToolsEditorMode/Source/ModelingToolsEditorMode/Public/ModelingToolsEditorModeToolkit.h
+++ b/Engine/Plugins/Editor/ModelingToolsEditorMode/Source/ModelingToolsEditorMode/Public/ModelingToolsEditorModeToolkit.h
@@ -117,6 +117,7 @@ private:
 
 	/** A utility function to register the tool palettes with the ToolkitBuilder */
 	void RegisterPalettes();
+	FDelegateHandle ActivePaletteChangedHandle;
 
 	TSharedPtr<SWidget> ToolShutdownViewportOverlayWidget;
 	void MakeToolShutdownOverlayWidget();
