ModelingMode: add support for Modeling Mode Extensions in new tool palette. Add FModelingToolsManagerCommands::RegisterExtensionPaletteCommand() to support dynamically registering FUICommandInfo's for section/tab buttons derived from extension info. Update Hair Modeling plugin to provide new optional tool palette command w/ icon new system.

diff --git a/Engine/Plugins/Editor/ModelingToolsEditorMode/Source/ModelingToolsEditorMode/Private/ModelingToolsEditorModeToolkit.cpp b/Engine/Plugins/Editor/ModelingToolsEditorMode/Source/ModelingToolsEditorMode/Private/ModelingToolsEditorModeToolkit.cpp
--- a/Engine/Plugins/Editor/ModelingToolsEditorMode/Source/ModelingToolsEditorMode/Private/ModelingToolsEditorModeToolkit.cpp
+++ b/Engine/Plugins/Editor/ModelingToolsEditorMode/Source/ModelingToolsEditorMode/Private/ModelingToolsEditorModeToolkit.cpp
@@ -457,6 +457,50 @@ void FModelingToolsEditorModeToolkit::RegisterPalettes()
 	ToolkitBuilder->SetActivePaletteOnLoad(Commands.LoadCreateTools.Get());
 	ToolkitBuilder->UpdateWidget();
 
+
+	if (IModularFeatures::Get().IsModularFeatureAvailable(IModelingModeToolExtension::GetModularFeatureName()))
+	{
+		TArray<IModelingModeToolExtension*> Extensions = IModularFeatures::Get().GetModularFeatureImplementations<IModelingModeToolExtension>(
+			IModelingModeToolExtension::GetModularFeatureName());
+		for (int32 k = 0; k < Extensions.Num(); ++k)
+		{
+			FText ExtensionName = Extensions[k]->GetExtensionName();
+			FText SectionName = Extensions[k]->GetToolSectionName();
+
+			FModelingModeExtensionExtendedInfo ExtensionExtendedInfo;
+			bool bHasExtendedInfo = Extensions[k]->GetExtensionExtendedInfo(ExtensionExtendedInfo);
+
+			TSharedPtr<FUICommandInfo> PaletteCommand;
+			if (bHasExtendedInfo && ExtensionExtendedInfo.ExtensionCommand.IsValid())
+			{
+				PaletteCommand = ExtensionExtendedInfo.ExtensionCommand;
+			}
+			else
+			{
+				FText UseTooltipText = (bHasExtendedInfo && ExtensionExtendedInfo.ToolPaletteButtonTooltip.IsEmpty() == false) ?
+					ExtensionExtendedInfo.ToolPaletteButtonTooltip : SectionName;
+				PaletteCommand = FModelingToolsManagerCommands::RegisterExtensionPaletteCommand(
+					FName(ExtensionName.ToString()),
+					SectionName, UseTooltipText, FSlateIcon());
+			}
+
+			TArray<TSharedPtr<FUICommandInfo>> PaletteItems;
+			FExtensionToolQueryInfo ExtensionQueryInfo;
+			ExtensionQueryInfo.bIsInfoQueryOnly = true;
+			TArray<FExtensionToolDescription> ToolSet;
+			Extensions[k]->GetExtensionTools(ExtensionQueryInfo, ToolSet);
+			for (const FExtensionToolDescription& ToolInfo : ToolSet)
+			{
+				PaletteItems.Add(ToolInfo.ToolCommand);
+			}
+
+			ToolkitBuilder->AddPalette(
+				MakeShareable(new FToolPalette(PaletteCommand.ToSharedRef(), PaletteItems)));
+		}
+	}
+
+
+
 	// if selected palette changes, make sure we are showing the palette command buttons, which may be hidden by active Tool
 	ActivePaletteChangedHandle = ToolkitBuilder->OnActivePaletteChanged.AddLambda([this]()
 	{
diff --git a/Engine/Plugins/Editor/ModelingToolsEditorMode/Source/ModelingToolsEditorMode/Private/ModelingToolsManagerActions.cpp b/Engine/Plugins/Editor/ModelingToolsEditorMode/Source/ModelingToolsEditorMode/Private/ModelingToolsManagerActions.cpp
--- a/Engine/Plugins/Editor/ModelingToolsEditorMode/Source/ModelingToolsEditorMode/Private/ModelingToolsManagerActions.cpp
+++ b/Engine/Plugins/Editor/ModelingToolsEditorMode/Source/ModelingToolsEditorMode/Private/ModelingToolsManagerActions.cpp
@@ -425,5 +425,43 @@ void FModelingToolsManagerCommands::RegisterCommands()
 }
 
 
+TSharedPtr<FUICommandInfo> FModelingToolsManagerCommands::RegisterExtensionPaletteCommand(
+	FName Name,
+	const FText& Label,
+	const FText& Tooltip,
+	const FSlateIcon& Icon)
+{
+	if (Instance.IsValid())
+	{
+		TSharedPtr<FModelingToolsManagerCommands> Commands = Instance.Pin();
+
+		for (FDynamicExtensionCommand& ExtensionCommand : Commands->ExtensionPaletteCommands)
+		{
+			if (ExtensionCommand.RegistrationName == Name)
+			{
+				return ExtensionCommand.Command;
+			}
+		}
+
+		TSharedPtr<FUICommandInfo> NewCommandInfo;
+
+		FUICommandInfo::MakeCommandInfo(
+			Commands->AsShared(),
+			NewCommandInfo,
+			Name, Label, Tooltip, Icon,
+			EUserInterfaceActionType::RadioButton,
+			FInputChord() );
+
+		FDynamicExtensionCommand NewExtensionCommand;
+		NewExtensionCommand.RegistrationName = Name;
+		NewExtensionCommand.Command = NewCommandInfo;
+		Commands->ExtensionPaletteCommands.Add(NewExtensionCommand);
+		return NewCommandInfo;
+	};
+
+	return TSharedPtr<FUICommandInfo>();
+}
+
+
 
 #undef LOCTEXT_NAMESPACE
diff --git a/Engine/Plugins/Editor/ModelingToolsEditorMode/Source/ModelingToolsEditorMode/Public/ModelingModeToolExtensions.h b/Engine/Plugins/Editor/ModelingToolsEditorMode/Source/ModelingToolsEditorMode/Public/ModelingModeToolExtensions.h
--- a/Engine/Plugins/Editor/ModelingToolsEditorMode/Source/ModelingToolsEditorMode/Public/ModelingModeToolExtensions.h
+++ b/Engine/Plugins/Editor/ModelingToolsEditorMode/Source/ModelingToolsEditorMode/Public/ModelingModeToolExtensions.h
@@ -46,6 +46,18 @@ struct MODELINGTOOLSEDITORMODE_API FExtensionToolDescription
 	UInteractiveToolBuilder* ToolBuilder = nullptr;
 };
 
+/**
+ * IModelingModeToolExtension implementations can optionally return additional information used in the Eidtor UI
+ * via instances of FModelingModeExtensionExtendedInfo.
+ */
+struct MODELINGTOOLSEDITORMODE_API FModelingModeExtensionExtendedInfo
+{
+	/** Tooltip to use for UI buttons that refer to the Extension */
+	FText ToolPaletteButtonTooltip;
+	/** Command button that will be used for the extension in the Modeling Mode Tab Bar. This can be undefined, in which case a Command button w/ default icon will be created. */
+	TSharedPtr<FUICommandInfo> ExtensionCommand;
+};
+
 /**
  * IModelingModeToolExtension uses the IModularFeature API to allow a Plugin to provide
  * a set of InteractiveTool's to be exposed in Modeling Mode. The Tools will be 
@@ -80,7 +92,11 @@ public:
 	 */
 	virtual void GetExtensionTools(const FExtensionToolQueryInfo& QueryInfo, TArray<FExtensionToolDescription>& ToolsOut) = 0;
 
-
+	/**
+	 * Query the Extension for extended information. This function is optional, the results will
+	 * be ignored unless it is overridden and returns true.
+	 */
+	virtual bool GetExtensionExtendedInfo(FModelingModeExtensionExtendedInfo& InfoOut) { return false; }
 
 	static FName GetModularFeatureName()
 	{
diff --git a/Engine/Plugins/Editor/ModelingToolsEditorMode/Source/ModelingToolsEditorMode/Public/ModelingToolsManagerActions.h b/Engine/Plugins/Editor/ModelingToolsEditorMode/Source/ModelingToolsEditorMode/Public/ModelingToolsManagerActions.h
--- a/Engine/Plugins/Editor/ModelingToolsEditorMode/Source/ModelingToolsEditorMode/Public/ModelingToolsManagerActions.h
+++ b/Engine/Plugins/Editor/ModelingToolsEditorMode/Source/ModelingToolsEditorMode/Public/ModelingToolsManagerActions.h
@@ -21,6 +21,13 @@ protected:
 	};
 	TArray<FStartToolCommand> RegisteredTools;		// Tool start-commands listed below are stored in this list
 
+	struct FDynamicExtensionCommand
+	{
+		FName RegistrationName;
+		TSharedPtr<FUICommandInfo> Command;
+	};
+	TArray<FDynamicExtensionCommand> ExtensionPaletteCommands;
+
 public:
 	/**
 	 * Find Tool start-command below by registered name (tool icon name in Mode palette)
@@ -205,6 +212,19 @@ public:
 	 * Initialize above commands
 	 */
 	virtual void RegisterCommands() override;
+
+
+	/**
+	 * Dynamically register a new UI Command based on the given Name/Label/Tooltip/Icon.
+	 * This is intended to be used outside of RegisterCommands(), eg by Modeling Mode extensions
+	 * loaded when the Toolkit is created.
+	 */
+	static TSharedPtr<FUICommandInfo> RegisterExtensionPaletteCommand(
+		FName Name,
+		const FText& Label,
+		const FText& Tooltip,
+		const FSlateIcon& Icon);
+
 };
 
 
diff --git a/Engine/Plugins/Experimental/HairModelingToolset/Source/HairModelingToolset/Private/HairModelingToolCommands.cpp b/Engine/Plugins/Experimental/HairModelingToolset/Source/HairModelingToolset/Private/HairModelingToolCommands.cpp
--- a/Engine/Plugins/Experimental/HairModelingToolset/Source/HairModelingToolset/Private/HairModelingToolCommands.cpp
+++ b/Engine/Plugins/Experimental/HairModelingToolset/Source/HairModelingToolset/Private/HairModelingToolCommands.cpp
@@ -22,6 +22,8 @@ FHairModelingToolCommands::FHairModelingToolCommands() :
 
 void FHairModelingToolCommands::RegisterCommands()
 {
+	UI_COMMAND(HairToolsTabButton, "Groom", "Groom Modeling Toolset", EUserInterfaceActionType::RadioButton, FInputChord());
+
 	UI_COMMAND(BeginGroomToMeshTool, "Helmet", "Generate Helmet Mesh for Selected Groom", EUserInterfaceActionType::ToggleButton, FInputChord());
 	UI_COMMAND(BeginGenerateLODMeshesTool, "HlmLOD", "Generate LODS for Hair Helmet", EUserInterfaceActionType::ToggleButton, FInputChord());
 	UI_COMMAND(BeginGroomCardsEditorTool, "CardsEd", "Edit Hair Cards", EUserInterfaceActionType::ToggleButton, FInputChord());
diff --git a/Engine/Plugins/Experimental/HairModelingToolset/Source/HairModelingToolset/Private/HairModelingToolsStyle.cpp b/Engine/Plugins/Experimental/HairModelingToolset/Source/HairModelingToolset/Private/HairModelingToolsStyle.cpp
--- a/Engine/Plugins/Experimental/HairModelingToolset/Source/HairModelingToolset/Private/HairModelingToolsStyle.cpp
+++ b/Engine/Plugins/Experimental/HairModelingToolset/Source/HairModelingToolset/Private/HairModelingToolsStyle.cpp
@@ -47,6 +47,9 @@ void FHairModelingToolsStyle::Initialize()
 	StyleSet->SetCoreContentRoot(FPaths::EngineContentDir() / TEXT("Slate"));
 
 	{
+		StyleSet->Set("HairModelingToolCommands.HairToolsTabButton", new IMAGE_BRUSH_SVG("Icons/HairHelmet", Icon20x20));
+		StyleSet->Set("HairModelingToolCommands.HairToolsTabButton.Small", new IMAGE_BRUSH_SVG("Icons/HairHelmet", Icon20x20));
+
 		StyleSet->Set("HairModelingToolCommands.BeginGroomCardsEditorTool", new IMAGE_BRUSH_SVG("Icons/CardsEditor", Icon20x20));
 		StyleSet->Set("HairModelingToolCommands.BeginGroomCardsEditorTool.Small", new IMAGE_BRUSH_SVG("Icons/CardsEditor", Icon20x20));
 		StyleSet->Set("HairModelingToolCommands.BeginGenerateLODMeshesTool", new IMAGE_BRUSH_SVG("Icons/GenLODs", Icon20x20));
diff --git a/Engine/Plugins/Experimental/HairModelingToolset/Source/HairModelingToolset/Private/HairModelingToolsetModule.cpp b/Engine/Plugins/Experimental/HairModelingToolset/Source/HairModelingToolset/Private/HairModelingToolsetModule.cpp
--- a/Engine/Plugins/Experimental/HairModelingToolset/Source/HairModelingToolset/Private/HairModelingToolsetModule.cpp
+++ b/Engine/Plugins/Experimental/HairModelingToolset/Source/HairModelingToolset/Private/HairModelingToolsetModule.cpp
@@ -58,7 +58,13 @@ void FHairModelingToolsetModule::GetExtensionTools(const FExtensionToolQueryInfo
 	GroomCardsEditorToolInfo.ToolCommand = FHairModelingToolCommands::Get().BeginGroomCardsEditorTool;
 	GroomCardsEditorToolInfo.ToolBuilder = NewObject<UGroomCardsEditorToolBuilder>();
 	ToolsOut.Add(GroomCardsEditorToolInfo);
+}
 
+bool FHairModelingToolsetModule::GetExtensionExtendedInfo(FModelingModeExtensionExtendedInfo& InfoOut)
+{
+	InfoOut.ExtensionCommand = FHairModelingToolCommands::Get().HairToolsTabButton;
+	InfoOut.ToolPaletteButtonTooltip = LOCTEXT("HairExtensionTooltip", "Tools for working with Groom meshes and cards");
+	return true;
 }
 
 
diff --git a/Engine/Plugins/Experimental/HairModelingToolset/Source/HairModelingToolset/Public/HairModelingToolCommands.h b/Engine/Plugins/Experimental/HairModelingToolset/Source/HairModelingToolset/Public/HairModelingToolCommands.h
--- a/Engine/Plugins/Experimental/HairModelingToolset/Source/HairModelingToolset/Public/HairModelingToolCommands.h
+++ b/Engine/Plugins/Experimental/HairModelingToolset/Source/HairModelingToolset/Public/HairModelingToolCommands.h
@@ -13,6 +13,8 @@ class HAIRMODELINGTOOLSET_API FHairModelingToolCommands : public TCommands<FHair
 public:
 	FHairModelingToolCommands();
 
+	TSharedPtr<FUICommandInfo> HairToolsTabButton;
+
 	TSharedPtr<FUICommandInfo> BeginGroomToMeshTool;
 	TSharedPtr<FUICommandInfo> BeginGenerateLODMeshesTool;
 	TSharedPtr<FUICommandInfo> BeginGroomCardsEditorTool;
diff --git a/Engine/Plugins/Experimental/HairModelingToolset/Source/HairModelingToolset/Public/HairModelingToolsetModule.h b/Engine/Plugins/Experimental/HairModelingToolset/Source/HairModelingToolset/Public/HairModelingToolsetModule.h
--- a/Engine/Plugins/Experimental/HairModelingToolset/Source/HairModelingToolset/Public/HairModelingToolsetModule.h
+++ b/Engine/Plugins/Experimental/HairModelingToolset/Source/HairModelingToolset/Public/HairModelingToolsetModule.h
@@ -18,4 +18,5 @@ public:
 	virtual FText GetExtensionName() override;
 	virtual FText GetToolSectionName() override;
 	virtual void GetExtensionTools(const FExtensionToolQueryInfo& QueryInfo, TArray<FExtensionToolDescription>& ToolsOut) override;
+	virtual bool GetExtensionExtendedInfo(FModelingModeExtensionExtendedInfo& InfoOut) override;
 };
