ModelingMode: avoid a crash on shutdown if hovering over a selection-enabled mesh and user alt-tabs out of the application, and then tabs back and closes it without updating the hover. The Selection behavior will still be in a hover-capture, but be destroyed in UModelingToolsEditorMode::Exit(), and then later when the InputRouter is shut down it is holding onto a pointer to an invalid InputBehavior.

diff --git a/Engine/Plugins/Editor/ModelingToolsEditorMode/Source/ModelingToolsEditorMode/Private/ModelingToolsEditorMode.cpp b/Engine/Plugins/Editor/ModelingToolsEditorMode/Source/ModelingToolsEditorMode/Private/ModelingToolsEditorMode.cpp
--- a/Engine/Plugins/Editor/ModelingToolsEditorMode/Source/ModelingToolsEditorMode/Private/ModelingToolsEditorMode.cpp
+++ b/Engine/Plugins/Editor/ModelingToolsEditorMode/Source/ModelingToolsEditorMode/Private/ModelingToolsEditorMode.cpp
@@ -1106,6 +1106,7 @@ void UModelingToolsEditorMode::Exit()
 	if (SelectionInteraction != nullptr)
 	{
 		SelectionInteraction->Shutdown();
+		GetInteractiveToolsContext()->InputRouter->ForceTerminateSource(SelectionInteraction);
 		GetInteractiveToolsContext()->InputRouter->DeregisterSource(SelectionInteraction);
 		SelectionInteraction = nullptr;
 	}
