Tool Presets: Prevents presets from serializing any UObjects that are not derived from the InteractiveToolPropertySets, fixing issues where tools would add themselves or other incompatible objects to their list of properties.

Additionally fixes a potential crash bug regarding the availability of the owning EditorMode for the toolkit.

diff --git a/Engine/Plugins/Editor/ModelingToolsEditorMode/Source/ModelingToolsEditorMode/Private/ModelingToolsEditorModeToolkit.cpp b/Engine/Plugins/Editor/ModelingToolsEditorMode/Source/ModelingToolsEditorMode/Private/ModelingToolsEditorModeToolkit.cpp
--- a/Engine/Plugins/Editor/ModelingToolsEditorMode/Source/ModelingToolsEditorMode/Private/ModelingToolsEditorModeToolkit.cpp
+++ b/Engine/Plugins/Editor/ModelingToolsEditorMode/Source/ModelingToolsEditorMode/Private/ModelingToolsEditorModeToolkit.cpp
@@ -948,7 +948,13 @@ TSharedPtr<SWidget> FModelingToolsEditorModeToolkit::MakePresetPanel()
 
 	const TSharedPtr<SHorizontalBox> NewContent = SNew(SHorizontalBox);
 	
-	auto IsToolActive = [this]() {return this->OwningEditorMode->GetToolManager()->GetActiveTool(EToolSide::Left) != nullptr ? EVisibility::Visible : EVisibility::Collapsed; };
+	auto IsToolActive = [this]() {
+		if (this->OwningEditorMode.IsValid())
+		{
+			return this->OwningEditorMode->GetToolManager()->GetActiveTool(EToolSide::Left) != nullptr ? EVisibility::Visible : EVisibility::Collapsed;
+		}
+		return EVisibility::Collapsed;
+	};
 
 	NewContent->AddSlot().HAlign(HAlign_Right)
 	[
@@ -1317,6 +1323,13 @@ void FModelingToolsEditorModeToolkit::CreateNewPresetInCollection(const FString&
 	[this, PresetLabel, ToolTip, Icon, CollectionPath](UInteractiveToolsPresetCollectionAsset& Preset, UInteractiveTool& Tool) {
 
 		TArray<UObject*> PropertySets = Tool.GetToolProperties();
+
+		// We only want to add the properties that are actual property sets, since the tool might have added other types of objects we don't
+		// want to serialize.
+		PropertySets.RemoveAll([this](UObject* Object) {
+			return Cast<UInteractiveToolPropertySet>(Object) == nullptr;
+		});
+
 		Preset.PerToolPresets.FindOrAdd(Tool.GetClass()->GetName()).ToolLabel = ActiveToolName;
 		if (ensure(ActiveToolIcon))
 		{
@@ -1356,6 +1369,13 @@ void FModelingToolsEditorModeToolkit::UpdatePresetInCollection(const FToolPreset
 	FModelingToolsEditorModeToolkitLocals::ExecuteWithPresetAndTool(*OwningEditorMode, EToolSide::Left, PresetToEditIn.PresetCollection,
 		[this, PresetToEditIn, bUpdateStoredPresetValues](UInteractiveToolsPresetCollectionAsset& Preset, UInteractiveTool& Tool) {
 			TArray<UObject*> PropertySets = Tool.GetToolProperties();
+
+			// We only want to add the properties that are actual property sets, since the tool might have added other types of objects we don't
+			// want to serialize.
+			PropertySets.RemoveAll([this](UObject* Object) {
+				return Cast<UInteractiveToolPropertySet>(Object) == nullptr;
+			});
+
 			if (!Preset.PerToolPresets.Contains(Tool.GetClass()->GetName()))
 			{
 				return;
