MeshModelingTools: Fix crashes when entering/exiting level instance editing, or when removing sublevels with modeling mode running.

diff --git a/Engine/Plugins/Editor/ModelingToolsEditorMode/Source/ModelingToolsEditorMode/Private/ModelingToolsEditorMode.cpp b/Engine/Plugins/Editor/ModelingToolsEditorMode/Source/ModelingToolsEditorMode/Private/ModelingToolsEditorMode.cpp
--- a/Engine/Plugins/Editor/ModelingToolsEditorMode/Source/ModelingToolsEditorMode/Private/ModelingToolsEditorMode.cpp
+++ b/Engine/Plugins/Editor/ModelingToolsEditorMode/Source/ModelingToolsEditorMode/Private/ModelingToolsEditorMode.cpp
@@ -136,6 +136,7 @@
 #include "DynamicMeshActor.h"
 #include "Components/BrushComponent.h"
 #include "Engine/StaticMeshActor.h"
+#include "Engine/World.h" // FWorldDelegates
 
 #include UE_INLINE_GENERATED_CPP_BY_NAME(ModelingToolsEditorMode)
 
@@ -1035,6 +1036,27 @@ void UModelingToolsEditorMode::Enter()
 			}
 		});
 	}
+
+	// Removing levels from the world can happen either by entering/exiting level instance edit mode, or
+	// by using the Levels panel. The problem is that any temporary actors we may have spawned in the 
+	// level for visualization, gizmos, etc. will be garbage collected. While EdModeInteractiveToolsContext
+	// should end the tools for us, we still have to take care of mode-level temporary actors.
+	FWorldDelegates::PreLevelRemovedFromWorld.AddWeakLambda(this, [this](ULevel*, UWorld*) {
+		// The ideal solution would be to just exit the mode, but we don't have a way to do that- we
+		// can only request a mode switch on next tick.Since this is too late to prevent a crash, we
+		// hand-clean up temporary actors here.
+		if (SelectionInteraction)
+		{
+			SelectionInteraction->Shutdown();
+		}
+
+		// Since we're doing this hand-cleanup above, we could actually register to OnCurrentLevelChanged and
+		// reinstate the temporary actors to stay in the mode. That seems a bit brittle, though, and there
+		// is still some hope that we can someday exit the mode instead of having to keep track of what is
+		// in danger of being garbage collected, so we might as well keep the workflow the same (i.e. exit
+		// mode).
+		GetModeManager()->ActivateDefaultMode();
+	});
 }
 
 void UModelingToolsEditorMode::RegisterUVEditor()
@@ -1102,6 +1124,8 @@ void UModelingToolsEditorMode::RegisterUVEditor()
 
 void UModelingToolsEditorMode::Exit()
 {
+	FWorldDelegates::PreLevelRemovedFromWorld.RemoveAll(this);
+
 	// shutdown selection interaction
 	if (SelectionInteraction != nullptr)
 	{
diff --git a/Engine/Source/Editor/Experimental/EditorInteractiveToolsFramework/Private/EdModeInteractiveToolsContext.cpp b/Engine/Source/Editor/Experimental/EditorInteractiveToolsFramework/Private/EdModeInteractiveToolsContext.cpp
--- a/Engine/Source/Editor/Experimental/EditorInteractiveToolsFramework/Private/EdModeInteractiveToolsContext.cpp
+++ b/Engine/Source/Editor/Experimental/EditorInteractiveToolsFramework/Private/EdModeInteractiveToolsContext.cpp
@@ -20,6 +20,7 @@
 #include "ShowFlags.h"				// for EngineShowFlags
 #include "Engine/Engine.h"
 #include "Engine/Selection.h"
+#include "Engine/World.h"
 #include "Misc/ITransaction.h"
 #include "ScopedTransaction.h"
 #include "Materials/Material.h"
@@ -409,6 +410,10 @@ void UEditorInteractiveToolsContext::TerminateActiveToolsOnWorldTearDown()
 {
 	DeactivateAllActiveTools(EToolShutdownType::Cancel);
 }
+void UEditorInteractiveToolsContext::TerminateActiveToolsOnLevelChange()
+{
+	DeactivateAllActiveTools(EToolShutdownType::Cancel);
+}
 
 void UEditorInteractiveToolsContext::PostInvalidation()
 {
@@ -1169,6 +1174,12 @@ void UModeManagerInteractiveToolsContext::Initialize(IToolsContextQueriesAPI* Qu
 		}
 	});
 
+	// Tools frequently spawn temporary objects in the level for visualization, gizmos, etc, so having a level be
+	// removed from the world puts the tools at risk of letting those objects be garbage collected and cause crashes
+	FWorldDelegates::PreLevelRemovedFromWorld.AddWeakLambda(this, [this](ULevel*, UWorld*) {
+		TerminateActiveToolsOnLevelChange();
+	});
+
 	ToolManager->OnToolEnded.AddUObject(this, &UModeManagerInteractiveToolsContext::OnToolEnded);
 	ToolManager->OnToolPostBuild.AddUObject(this, &UModeManagerInteractiveToolsContext::OnToolPostBuild);
 
@@ -1195,6 +1206,8 @@ void UModeManagerInteractiveToolsContext::Shutdown()
 		GEditor->OnViewportClientListChanged().Remove(ViewportClientListChangedHandle);
 	}
 
+	FWorldDelegates::PreLevelRemovedFromWorld.RemoveAll(this);
+
 	Super::Shutdown();
 }
 
diff --git a/Engine/Source/Editor/Experimental/EditorInteractiveToolsFramework/Public/EdModeInteractiveToolsContext.h b/Engine/Source/Editor/Experimental/EditorInteractiveToolsFramework/Public/EdModeInteractiveToolsContext.h
--- a/Engine/Source/Editor/Experimental/EditorInteractiveToolsFramework/Public/EdModeInteractiveToolsContext.h
+++ b/Engine/Source/Editor/Experimental/EditorInteractiveToolsFramework/Public/EdModeInteractiveToolsContext.h
@@ -80,9 +80,12 @@ public:
 	// default behavior is to accept active tool
 	virtual void TerminateActiveToolsOnSaveWorld();
 
-	// default behavior is to accept active tool
+	// default behavior is to cancel active tool
 	virtual void TerminateActiveToolsOnWorldTearDown();
 
+	// default behavior is to cancel active tool
+	virtual void TerminateActiveToolsOnLevelChange();
+
 	FEditorModeTools* GetParentEditorModeManager() const { return EditorModeManager; }
 
 	IToolsContextQueriesAPI* GetQueriesAPI() const { return QueriesAPI; }
