remove stylus plugin from modeling mode and comment out related code, so it could in theory be brought back in the future

diff --git a/Engine/Plugins/Editor/ModelingToolsEditorMode/ModelingToolsEditorMode.uplugin b/Engine/Plugins/Editor/ModelingToolsEditorMode/ModelingToolsEditorMode.uplugin
--- a/Engine/Plugins/Editor/ModelingToolsEditorMode/ModelingToolsEditorMode.uplugin
+++ b/Engine/Plugins/Editor/ModelingToolsEditorMode/ModelingToolsEditorMode.uplugin
@@ -36,10 +36,6 @@
 			"Enabled": true,
 			"TargetAllowList": [ "Editor" ]
 		},
-		{
-			"Name": "StylusInput",
-			"Enabled": true
-		},
 		{
 			"Name": "ToolPresets",
 			"Enabled": true,
diff --git a/Engine/Plugins/Editor/ModelingToolsEditorMode/Source/ModelingToolsEditorMode/ModelingToolsEditorMode.Build.cs b/Engine/Plugins/Editor/ModelingToolsEditorMode/Source/ModelingToolsEditorMode/ModelingToolsEditorMode.Build.cs
--- a/Engine/Plugins/Editor/ModelingToolsEditorMode/Source/ModelingToolsEditorMode/ModelingToolsEditorMode.Build.cs
+++ b/Engine/Plugins/Editor/ModelingToolsEditorMode/Source/ModelingToolsEditorMode/ModelingToolsEditorMode.Build.cs
@@ -47,7 +47,6 @@ public class ModelingToolsEditorMode : ModuleRules
 				"EditorWidgets",
 				"ModelingEditorUI",
 				"WidgetRegistration",
-				"StylusInput",
 				"DeveloperSettings",
 				"PropertyEditor",
 				"ToolMenus",
diff --git a/Engine/Plugins/Editor/ModelingToolsEditorMode/Source/ModelingToolsEditorMode/Private/ModelingToolsEditorMode.cpp b/Engine/Plugins/Editor/ModelingToolsEditorMode/Source/ModelingToolsEditorMode/Private/ModelingToolsEditorMode.cpp
--- a/Engine/Plugins/Editor/ModelingToolsEditorMode/Source/ModelingToolsEditorMode/Private/ModelingToolsEditorMode.cpp
+++ b/Engine/Plugins/Editor/ModelingToolsEditorMode/Source/ModelingToolsEditorMode/Private/ModelingToolsEditorMode.cpp
@@ -10,7 +10,6 @@
 #include "ModelingToolsEditorModeToolkit.h"
 #include "ILevelEditor.h"
 #include "ModelingToolsEditorModeSettings.h"
-#include "IStylusState.h"
 #include "ToolTargetManager.h"
 #include "ContextObjectStore.h"
 #include "InputRouter.h"
@@ -117,8 +116,14 @@
 #include "EditorModeManager.h"
 #include "UnrealWidget.h"
 
-// stylus support
+// Stylus support is currently disabled due to issues with the stylus plugin
+// We are leaving the code in this cpp file, defined out, so that it is easier to bring back if/when the stylus plugin is improved.
+#define ENABLE_STYLUS_SUPPORT 0
+
+#if ENABLE_STYLUS_SUPPORT 
+#include "IStylusState.h"
 #include "IStylusInputModule.h"
+#endif
 
 #include "LevelEditor.h"
 #include "SLevelViewport.h"
@@ -270,6 +275,8 @@ void UModelingToolsEditorMode::Tick(FEditorViewportClient* ViewportClient, float
 	}
 }
 
+// Note: Stylus support is currently non-functioning; the code to enable it is left here as reference in case it is brought back
+#if ENABLE_STYLUS_SUPPORT
 //
 // FStylusStateTracker registers itself as a listener for stylus events and implements
 // the IToolStylusStateProviderAPI interface, which allows MeshSurfacePointTool implementations
@@ -356,7 +363,7 @@ public:
 	}
 
 };
-
+#endif // ENABLE_STYLUS_SUPPORT
 
 
 
@@ -394,8 +401,10 @@ void UModelingToolsEditorMode::Enter()
 	GetInteractiveToolsContext()->OnRender.AddUObject(this, &UModelingToolsEditorMode::OnToolsContextRender);
 	GetInteractiveToolsContext()->OnDrawHUD.AddUObject(this, &UModelingToolsEditorMode::OnToolsContextDrawHUD);
 
+#if ENABLE_STYLUS_SUPPORT 
 	// register stylus event handler
 	StylusStateTracker = MakeUnique<FStylusStateTracker>();
+#endif
 
 	// register gizmo helper
 	UE::TransformGizmoUtil::RegisterTransformGizmoContextObject(GetInteractiveToolsContext());
@@ -593,11 +602,15 @@ void UModelingToolsEditorMode::Enter()
 	//
 
 	auto MoveVerticesToolBuilder = NewObject<UMeshVertexSculptToolBuilder>();
+#if ENABLE_STYLUS_SUPPORT 
 	MoveVerticesToolBuilder->StylusAPI = StylusStateTracker.Get();
+#endif
 	RegisterTool(ToolManagerCommands.BeginSculptMeshTool, TEXT("BeginSculptMeshTool"), MoveVerticesToolBuilder);
 
 	auto MeshGroupPaintToolBuilder = NewObject<UMeshGroupPaintToolBuilder>();
+#if ENABLE_STYLUS_SUPPORT 
 	MeshGroupPaintToolBuilder->StylusAPI = StylusStateTracker.Get();
+#endif
 	RegisterTool(ToolManagerCommands.BeginMeshGroupPaintTool, TEXT("BeginMeshGroupPaintTool"), MeshGroupPaintToolBuilder);
 
 	RegisterTool(ToolManagerCommands.BeginPolyEditTool, TEXT("BeginPolyEditTool"), NewObject<UEditMeshPolygonsToolBuilder>());
@@ -653,7 +666,9 @@ void UModelingToolsEditorMode::Enter()
 
 	auto DynaSculptToolBuilder = NewObject<UDynamicMeshSculptToolBuilder>();
 	DynaSculptToolBuilder->bEnableRemeshing = true;
+#if ENABLE_STYLUS_SUPPORT 
 	DynaSculptToolBuilder->StylusAPI = StylusStateTracker.Get();
+#endif
 	RegisterTool(ToolManagerCommands.BeginRemeshSculptMeshTool, TEXT("BeginRemeshSculptMeshTool"), DynaSculptToolBuilder);
 
 	RegisterTool(ToolManagerCommands.BeginRemeshMeshTool, TEXT("BeginRemeshMeshTool"), NewObject<URemeshMeshToolBuilder>());
@@ -1180,7 +1195,9 @@ void UModelingToolsEditorMode::Exit()
 		FEngineAnalytics::GetProvider().RecordEvent(TEXT("Editor.Usage.MeshModelingMode.Exit"), Attributes);
 	}
 
+#if ENABLE_STYLUS_SUPPORT 
 	StylusStateTracker = nullptr;
+#endif
 
 	// TODO: cannot deregister currently because if another mode is also registering, its Enter()
 	// will be called before our Exit()
diff --git a/Engine/Plugins/Editor/ModelingToolsEditorMode/Source/ModelingToolsEditorMode/Public/ModelingToolsEditorMode.h b/Engine/Plugins/Editor/ModelingToolsEditorMode/Source/ModelingToolsEditorMode/Public/ModelingToolsEditorMode.h
--- a/Engine/Plugins/Editor/ModelingToolsEditorMode/Source/ModelingToolsEditorMode/Public/ModelingToolsEditorMode.h
+++ b/Engine/Plugins/Editor/ModelingToolsEditorMode/Source/ModelingToolsEditorMode/Public/ModelingToolsEditorMode.h
@@ -113,7 +113,8 @@ protected:
 	FDelegateHandle EditorClosedEventHandle;
 	void OnEditorClosed();
 
-	TUniquePtr<FStylusStateTracker> StylusStateTracker;
+	// Stylus support is currently disabled; this is left in for reference if/when it is brought back
+	//TUniquePtr<FStylusStateTracker> StylusStateTracker;
 
 	TSharedPtr<FLevelObjectsObserver> LevelObjectsObserver;
 
