Tool Presets: Handle scenario when the user provides a blank label and tooltip during preset creation better.

diff --git a/Engine/Plugins/Editor/ModelingToolsEditorMode/Source/ModelingToolsEditorMode/Private/ModelingToolsEditorModeToolkit.cpp b/Engine/Plugins/Editor/ModelingToolsEditorMode/Source/ModelingToolsEditorMode/Private/ModelingToolsEditorModeToolkit.cpp
--- a/Engine/Plugins/Editor/ModelingToolsEditorMode/Source/ModelingToolsEditorMode/Private/ModelingToolsEditorModeToolkit.cpp
+++ b/Engine/Plugins/Editor/ModelingToolsEditorMode/Source/ModelingToolsEditorMode/Private/ModelingToolsEditorModeToolkit.cpp
@@ -1003,6 +1003,9 @@ TSharedRef<SWidget> FModelingToolsEditorModeToolkit::GetPresetCreateButtonConten
 
 		auto OpenNewPresetDialog = [this]()
 	{
+		NewPresetLabel.Empty();
+		NewPresetTooltip.Empty();
+
 		// Set the result if they just click Ok
 		SGenericDialogWidget::FArguments FolderDialogArguments;
 		FolderDialogArguments.OnOkPressed_Lambda([this]()
@@ -1340,7 +1343,14 @@ void FModelingToolsEditorModeToolkit::CreateNewPresetInCollection(const FString&
 			Preset.PerToolPresets.FindOrAdd(Tool.GetClass()->GetName()).ToolIcon = *ActiveToolIcon;
 		}
 		FInteractiveToolPresetDefintion PresetValuesToCreate;
-		PresetValuesToCreate.Label = PresetLabel;
+		if (PresetLabel.IsEmpty())
+		{
+			PresetValuesToCreate.Label = FString::Printf(TEXT("Unnamed_Preset-%d"), Preset.PerToolPresets.FindOrAdd(Tool.GetClass()->GetName()).NamedPresets.Num()+1);
+		}
+		else
+		{
+			PresetValuesToCreate.Label = PresetLabel;
+		}
 		PresetValuesToCreate.Tooltip = ToolTip;
 		//PresetValuesToCreate.Icon = Icon;
 
