Tool Presets: Fix handling of tool preset collection selection to properly display collections after switching back and forth from no collection selection to a specific collection selected.

diff --git a/Engine/Plugins/Editor/ModelingToolsEditorMode/Source/ModelingToolsEditorMode/Private/ModelingToolsEditorModeToolkit.cpp b/Engine/Plugins/Editor/ModelingToolsEditorMode/Source/ModelingToolsEditorMode/Private/ModelingToolsEditorModeToolkit.cpp
--- a/Engine/Plugins/Editor/ModelingToolsEditorMode/Source/ModelingToolsEditorMode/Private/ModelingToolsEditorModeToolkit.cpp
+++ b/Engine/Plugins/Editor/ModelingToolsEditorMode/Source/ModelingToolsEditorMode/Private/ModelingToolsEditorModeToolkit.cpp
@@ -166,6 +166,7 @@ FModelingToolsEditorModeToolkit::FModelingToolsEditorModeToolkit()
 	UPresetUserSettings::Get()->LoadEditorConfig();
 
 	RecentPresetCollectionProvider = MakeShared< FRecentPresetCollectionProvider>();
+	CurrentPreset = MakeShared<FAssetData>();
 }
 
 FModelingToolsEditorModeToolkit::~FModelingToolsEditorModeToolkit()
@@ -636,7 +637,7 @@ void FModelingToolsEditorModeToolkit::Init(const TSharedPtr<IToolkitHost>& InitT
 	MakeSelectionPaletteOverlayWidget();
 	GetToolkitHost()->AddViewportOverlayWidget(SelectionPaletteOverlayWidget.ToSharedRef());
 
-	CurrentPreset = FSoftObjectPath(); // Default to the default collection by leaving this null.
+	CurrentPresetPath = FSoftObjectPath(); // Default to the default collection by leaving this null.
 }
 
 void FModelingToolsEditorModeToolkit::MakeToolShutdownOverlayWidget()
@@ -993,7 +994,7 @@ TSharedPtr<SWidget> FModelingToolsEditorModeToolkit::MakePresetPanel()
 
 bool FModelingToolsEditorModeToolkit::IsPresetEnabled() const
 {
-	return CurrentPreset.IsAsset();
+	return CurrentPresetPath.IsAsset();
 }
 
 
@@ -1011,7 +1012,7 @@ TSharedRef<SWidget> FModelingToolsEditorModeToolkit::GetPresetCreateButtonConten
 		FolderDialogArguments.OnOkPressed_Lambda([this]()
 			{
 				CreateNewPresetInCollection(NewPresetLabel,
-					CurrentPreset,
+					CurrentPresetPath,
 					NewPresetTooltip,
 					NewPresetIcon);
 			});
@@ -1070,6 +1071,7 @@ TSharedRef<SWidget> FModelingToolsEditorModeToolkit::GetPresetCreateButtonConten
 					.OnSelectionChanged(this, &FModelingToolsEditorModeToolkit::HandlePresetAssetChanged)
 					.ToolTipText(LOCTEXT("ToolPresets_CreatePresetCollection_Tooltip", "The asset in which to store this new preset."))
 					//.RecentAssetsProvider(RecentPresetCollectionProvider) // TODO: Improve this widget before enabling this feature
+					.InitiallySelectedAsset(*CurrentPreset)
 					.FlyoutTileSize(FVector2D(80, 80))
 					.ComboButtonTileSize(FVector2D(80, 80))
 				]
@@ -1286,12 +1288,14 @@ void FModelingToolsEditorModeToolkit::RebuildPresetListForTool(bool bSettingsOpe
 
 void FModelingToolsEditorModeToolkit::HandlePresetAssetChanged(const FAssetData& InAssetData)
 {
-	CurrentPreset.SetPath(InAssetData.GetObjectPathString());
+	CurrentPresetPath.SetPath(InAssetData.GetObjectPathString());
+	CurrentPresetLabel = FText();
+	*CurrentPreset = InAssetData;
 
 	UInteractiveToolsPresetCollectionAsset* Preset = nullptr;
-	if (CurrentPreset.IsAsset())
+	if (CurrentPresetPath.IsAsset())
 	{
-		Preset = Cast<UInteractiveToolsPresetCollectionAsset>(CurrentPreset.TryLoad());
+		Preset = Cast<UInteractiveToolsPresetCollectionAsset>(CurrentPresetPath.TryLoad());
 	}
 	if (Preset)
 	{
@@ -1326,7 +1330,7 @@ void FModelingToolsEditorModeToolkit::LoadPresetFromCollection(const int32 Prese
 
 void FModelingToolsEditorModeToolkit::CreateNewPresetInCollection(const FString& PresetLabel, FSoftObjectPath CollectionPath, const FString& ToolTip, FSlateIcon Icon)
 {
-	FModelingToolsEditorModeToolkitLocals::ExecuteWithPresetAndTool(*OwningEditorMode, EToolSide::Left, CurrentPreset,
+	FModelingToolsEditorModeToolkitLocals::ExecuteWithPresetAndTool(*OwningEditorMode, EToolSide::Left, CurrentPresetPath,
 	[this, PresetLabel, ToolTip, Icon, CollectionPath](UInteractiveToolsPresetCollectionAsset& Preset, UInteractiveTool& Tool) {
 
 		TArray<UObject*> PropertySets = Tool.GetToolProperties();
diff --git a/Engine/Plugins/Editor/ModelingToolsEditorMode/Source/ModelingToolsEditorMode/Public/ModelingToolsEditorModeToolkit.h b/Engine/Plugins/Editor/ModelingToolsEditorMode/Source/ModelingToolsEditorMode/Public/ModelingToolsEditorModeToolkit.h
--- a/Engine/Plugins/Editor/ModelingToolsEditorMode/Source/ModelingToolsEditorMode/Public/ModelingToolsEditorModeToolkit.h
+++ b/Engine/Plugins/Editor/ModelingToolsEditorMode/Source/ModelingToolsEditorMode/Public/ModelingToolsEditorModeToolkit.h
@@ -19,6 +19,7 @@ class STextBlock;
 class UGeometrySelectionManager;
 class UInteractiveToolsPresetCollectionAsset;
 class FRecentPresetCollectionProvider;
+struct FAssetData;
 
 struct FToolPresetOption
 {
@@ -158,7 +159,8 @@ private:
 	// Presets
 	TSharedPtr<SWidget> ToolPresetArea;
 	TSharedPtr<SWidget> MakePresetPanel();
-	FSoftObjectPath CurrentPreset;
+	TSharedPtr<FAssetData> CurrentPreset;
+	FSoftObjectPath CurrentPresetPath;
 	FText CurrentPresetLabel;
 	FToolPresetOption SelectedPresetToEdit;
 	FString NewPresetLabel;
@@ -175,7 +177,7 @@ private:
 	void LoadPresetFromCollection(const int32 PresetIndex, FSoftObjectPath CollectionPath);
 	void UpdatePresetInCollection(const FToolPresetOption& PresetToEditIn, bool bUpdateStoredPresetValues);
 
-	FString GetCurrentPresetPath() { return CurrentPreset.GetAssetPathString(); }
+	FString GetCurrentPresetPath() { return CurrentPresetPath.GetAssetPathString(); }
 	void HandlePresetAssetChanged(const FAssetData& InAssetData);
 	bool HandleFilterPresetAsset(const FAssetData& InAssetData);
 	
