MeshModelingToolset: Re-enable ISM tools in Restrictive Mode

diff --git a/Engine/Plugins/Editor/ModelingToolsEditorMode/Source/ModelingToolsEditorMode/Private/ModelingToolsEditorMode.cpp b/Engine/Plugins/Editor/ModelingToolsEditorMode/Source/ModelingToolsEditorMode/Private/ModelingToolsEditorMode.cpp
--- a/Engine/Plugins/Editor/ModelingToolsEditorMode/Source/ModelingToolsEditorMode/Private/ModelingToolsEditorMode.cpp
+++ b/Engine/Plugins/Editor/ModelingToolsEditorMode/Source/ModelingToolsEditorMode/Private/ModelingToolsEditorMode.cpp
@@ -644,7 +644,7 @@ void UModelingToolsEditorMode::Enter()
 	RegisterTool(ToolManagerCommands.BeginSubdividePolyTool, TEXT("BeginSubdividePolyTool"), NewObject<USubdividePolyToolBuilder>());
 
 	UPatternToolBuilder* PatternToolBuilder = NewObject<UPatternToolBuilder>();
-	PatternToolBuilder->bEnableCreateISMCs = !(ModelingModeSettings && ModelingModeSettings->InRestrictiveMode());
+	PatternToolBuilder->bEnableCreateISMCs = true;
 	RegisterTool(ToolManagerCommands.BeginPatternTool, TEXT("BeginPatternTool"), PatternToolBuilder);
 
 	UCombineMeshesToolBuilder* CombineMeshesToolBuilder = NewObject<UCombineMeshesToolBuilder>();
@@ -662,12 +662,8 @@ void UModelingToolsEditorMode::Enter()
 	GenerateSMLODToolBuilder->bInRestrictiveMode = ModelingModeSettings && ModelingModeSettings->InRestrictiveMode();
 	RegisterTool(ToolManagerCommands.BeginGenerateStaticMeshLODAssetTool, TEXT("BeginGenerateStaticMeshLODAssetTool"), GenerateSMLODToolBuilder);
 
-	// ISMEd is disabled in Restrictive Mode.
-	if (ToolManagerCommands.BeginISMEditorTool)
-	{
-		UISMEditorToolBuilder* ISMEditorToolBuilder = NewObject<UISMEditorToolBuilder>();
-		RegisterTool(ToolManagerCommands.BeginISMEditorTool, TEXT("BeginISMEditorTool"), ISMEditorToolBuilder);
-	}
+	UISMEditorToolBuilder* ISMEditorToolBuilder = NewObject<UISMEditorToolBuilder>();
+	RegisterTool(ToolManagerCommands.BeginISMEditorTool, TEXT("BeginISMEditorTool"), ISMEditorToolBuilder);
 
 
 	// edit tools
diff --git a/Engine/Plugins/Editor/ModelingToolsEditorMode/Source/ModelingToolsEditorMode/Private/ModelingToolsEditorModeToolkit.cpp b/Engine/Plugins/Editor/ModelingToolsEditorMode/Source/ModelingToolsEditorMode/Private/ModelingToolsEditorModeToolkit.cpp
--- a/Engine/Plugins/Editor/ModelingToolsEditorMode/Source/ModelingToolsEditorMode/Private/ModelingToolsEditorModeToolkit.cpp
+++ b/Engine/Plugins/Editor/ModelingToolsEditorMode/Source/ModelingToolsEditorMode/Private/ModelingToolsEditorModeToolkit.cpp
@@ -302,12 +302,10 @@ void FModelingToolsEditorModeToolkit::RegisterPalettes()
 		Commands.BeginConvertMeshesTool,
 
 		Commands.BeginSplitMeshesTool,
-		Commands.BeginPatternTool
+		Commands.BeginPatternTool,
+		
+		Commands.BeginISMEditorTool
 	});
-	if (Commands.BeginISMEditorTool)
-	{
-		TransformPaletteItems.Add(Commands.BeginISMEditorTool);
-	}
 	ToolkitBuilder->AddPalette(
 		MakeShareable( new FToolPalette( Commands.LoadTransformTools.ToSharedRef(), TransformPaletteItems ) ) );
 
@@ -1743,10 +1741,7 @@ void FModelingToolsEditorModeToolkit::BuildToolPalette(FName PaletteIndex, class
 	{
 		ToolbarBuilder.AddToolBarButton(Commands.BeginLODManagerTool);
 		ToolbarBuilder.AddToolBarButton(Commands.BeginGenerateStaticMeshLODAssetTool);
-		if (Commands.BeginISMEditorTool)
-		{
-			ToolbarBuilder.AddToolBarButton(Commands.BeginISMEditorTool);
-		}
+		ToolbarBuilder.AddToolBarButton(Commands.BeginISMEditorTool);
 	}
 	else if (PaletteIndex == VoxToolsTabName)
 	{
diff --git a/Engine/Plugins/Editor/ModelingToolsEditorMode/Source/ModelingToolsEditorMode/Private/ModelingToolsManagerActions.cpp b/Engine/Plugins/Editor/ModelingToolsEditorMode/Source/ModelingToolsEditorMode/Private/ModelingToolsManagerActions.cpp
--- a/Engine/Plugins/Editor/ModelingToolsEditorMode/Source/ModelingToolsEditorMode/Private/ModelingToolsManagerActions.cpp
+++ b/Engine/Plugins/Editor/ModelingToolsEditorMode/Source/ModelingToolsEditorMode/Private/ModelingToolsManagerActions.cpp
@@ -315,19 +315,13 @@ void FModelingToolsManagerCommands::RegisterCommands()
 	{
 		REGISTER_MODELING_TOOL_COMMAND(BeginLODManagerTool, "LODMgr", "Inspect the LODs of a Static Mesh Asset");
 		REGISTER_MODELING_TOOL_COMMAND(BeginGenerateStaticMeshLODAssetTool, "AutoLOD", "Automatically generate a simplified LOD with baked Textures/Materials for a Mesh");
-		if (!Settings->InRestrictiveMode())
-		{
-			REGISTER_MODELING_TOOL_COMMAND(BeginISMEditorTool, "ISMEd", "Edit the Instances of Instanced Static Mesh Components");
-		}
+		REGISTER_MODELING_TOOL_COMMAND(BeginISMEditorTool, "ISMEd", "Edit the Instances of Instanced Static Mesh Components");
 	}
 	else
 	{
 		REGISTER_MODELING_TOOL_COMMAND(BeginLODManagerTool, "LOD Manager", "Inspect the LODs of a Static Mesh Asset");
 		REGISTER_MODELING_TOOL_COMMAND(BeginGenerateStaticMeshLODAssetTool, "AutoLOD", "Automatically generate a simplified LOD with baked Textures/Materials for a Mesh");
-		if (!Settings->InRestrictiveMode())
-		{
-			REGISTER_MODELING_TOOL_COMMAND(BeginISMEditorTool, "ISM Editor", "Edit the Instances of Instanced Static Mesh Components");
-		}
+		REGISTER_MODELING_TOOL_COMMAND(BeginISMEditorTool, "ISM Editor", "Edit the Instances of Instanced Static Mesh Components");
 	}
 
 	REGISTER_MODELING_TOOL_COMMAND(BeginAddPatchTool, "Patch", "Add Patch");
