MeshModelingTools: Conditionally turn the accept/cancel overlay opaque to draw on top of level instance editing overlay. Also fix minor memory leak in the way we did custom section colors in the old palette.

diff --git a/Engine/Plugins/Editor/ModelingToolsEditorMode/Source/ModelingToolsEditorMode/Private/ModelingToolsEditorModeStyle.cpp b/Engine/Plugins/Editor/ModelingToolsEditorMode/Source/ModelingToolsEditorMode/Private/ModelingToolsEditorModeStyle.cpp
--- a/Engine/Plugins/Editor/ModelingToolsEditorMode/Source/ModelingToolsEditorMode/Private/ModelingToolsEditorModeStyle.cpp
+++ b/Engine/Plugins/Editor/ModelingToolsEditorMode/Source/ModelingToolsEditorMode/Private/ModelingToolsEditorModeStyle.cpp
@@ -463,6 +463,22 @@ void FModelingToolsEditorModeStyle::Initialize()
 		StyleSet->Set("PolyEd.SelectionToolbar.ToggleButton.End", ToggleButtonEnd);
 	}
 
+	// Style to be applied to customizable section headers so that the color shows up properly
+	{
+		// look up default radii for palette toolbar expandable area headers
+		FVector4 HeaderRadii(4, 4, 0, 0);
+		const FSlateBrush* BaseBrush = FAppStyle::Get().GetBrush("PaletteToolbar.ExpandableAreaHeader");
+		if (BaseBrush != nullptr)
+		{
+			HeaderRadii = BaseBrush->OutlineSettings.CornerRadii;
+		}
+		StyleSet->Set("ModelingMode.WhiteExpandableAreaHeader", new FSlateRoundedBoxBrush(FSlateColor(FLinearColor::White), HeaderRadii));
+	}
+
+	// Similar to EditorViewport.OverlayBrush, but opaque with a gray color to be able to be placed on top of other overlays.
+	StyleSet->Set("ModelingMode.OpaqueOverlayBrush", 
+		new FSlateRoundedBoxBrush(FStyleColors::Panel.GetSpecifiedColor(), 8.0, FStyleColors::Dropdown, 1.0));
+
 	FSlateStyleRegistry::RegisterSlateStyle(*StyleSet.Get());
 };
 
diff --git a/Engine/Plugins/Editor/ModelingToolsEditorMode/Source/ModelingToolsEditorMode/Private/ModelingToolsEditorModeToolkit.cpp b/Engine/Plugins/Editor/ModelingToolsEditorMode/Source/ModelingToolsEditorMode/Private/ModelingToolsEditorModeToolkit.cpp
--- a/Engine/Plugins/Editor/ModelingToolsEditorMode/Source/ModelingToolsEditorMode/Private/ModelingToolsEditorModeToolkit.cpp
+++ b/Engine/Plugins/Editor/ModelingToolsEditorMode/Source/ModelingToolsEditorMode/Private/ModelingToolsEditorModeToolkit.cpp
@@ -2,6 +2,7 @@
 
 #include "ModelingToolsEditorModeToolkit.h"
 #include "Brushes/SlateRoundedBoxBrush.h"
+#include "EditorModeManager.h"
 #include "IGeometryProcessingInterfacesModule.h"
 #include "GeometryProcessingInterfaces/IUVEditorModularFeature.h"
 #include "EdModeInteractiveToolsContext.h"
@@ -640,6 +641,15 @@ void FModelingToolsEditorModeToolkit::Init(const TSharedPtr<IToolkitHost>& InitT
 
 void FModelingToolsEditorModeToolkit::MakeToolShutdownOverlayWidget()
 {
+	const FSlateBrush* OverlayBrush = FAppStyle::Get().GetBrush("EditorViewport.OverlayBrush");
+	// If there is another mode, it might also have an overlay, and we would like ours to be opaque in that case
+	// to draw on top cleanly (e.g., level instance editing mode has an overlay in the same place. Note that level
+	// instance mode currently marks itself as not visible despite the overlay, so we shouldn't use IsOnlyVisibleActiveMode)
+	if (!GetEditorModeManager().IsOnlyActiveMode(UModelingToolsEditorMode::EM_ModelingToolsEditorModeId))
+	{
+		OverlayBrush = FModelingToolsEditorModeStyle::Get()->GetBrush("ModelingMode.OpaqueOverlayBrush");
+	}
+
 	SAssignNew(ToolShutdownViewportOverlayWidget, SHorizontalBox)
 
 	+SHorizontalBox::Slot()
@@ -648,7 +658,7 @@ void FModelingToolsEditorModeToolkit::MakeToolShutdownOverlayWidget()
 	.Padding(FMargin(0.0f, 0.0f, 0.f, 15.f))
 	[
 		SNew(SBorder)
-		.BorderImage(FAppStyle::Get().GetBrush("EditorViewport.OverlayBrush"))
+		.BorderImage(OverlayBrush)
 		.Padding(8.f)
 		[
 			SNew(SHorizontalBox)
@@ -1896,14 +1906,6 @@ void FModelingToolsEditorModeToolkit::InvokeUI()
 
 	UModelingToolsModeCustomizationSettings* UISettings = GetMutableDefault<UModelingToolsModeCustomizationSettings>();
 
-	// look up default radii for palette toolbar expandable area headers
-	FVector4 HeaderRadii(4, 4, 0, 0);
-	const FSlateBrush* BaseBrush = FAppStyle::Get().GetBrush("PaletteToolbar.ExpandableAreaHeader");
-	if (BaseBrush != nullptr)
-	{
-		HeaderRadii = BaseBrush->OutlineSettings.CornerRadii;
-	}
-
 	// Generate a map for tool specific colors
 	TMap<FString, FLinearColor> SectionIconColorMap;
 	TMap<FString, TMap<FString, FLinearColor>> SectionToolIconColorMap;
@@ -1952,7 +1954,13 @@ void FModelingToolsEditorModeToolkit::InvokeUI()
 							TSharedPtr<SBorder> TopBorder = StaticCastSharedPtr<SBorder>(SlotWidgetPtr);
 							if (TopBorder.IsValid())
 							{
-								TopBorder->SetBorderImage(new FSlateRoundedBoxBrush(FSlateColor(ToolColor.Color), HeaderRadii));
+								// The border image needs to be swapped to a white one so that the color applied to it is
+								// not darkened by the usual brush. Note that we can't just create a new FSlateRoundedBoxBrush
+								// with the proper color in-line because a brush needs to be associated with a style set that
+								// is responsible for freeing it, else it will leak (or we could own the brush in the toolkit,
+								// but this is frowned upon).
+								TopBorder->SetBorderImage(FModelingToolsEditorModeStyle::Get()->GetBrush("ModelingMode.WhiteExpandableAreaHeader"));
+								TopBorder->SetBorderBackgroundColor(FSlateColor(ToolColor.Color));
 							}
 						}
 					}
