Lambda capture fixes ahead of perfect forwarding of lambda Execute*() arguments.

diff --git a/Engine/Plugins/Animation/ControlRig/Source/ControlRigEditor/Private/EditMode/SControlRigEditModeTools.cpp b/Engine/Plugins/Animation/ControlRig/Source/ControlRigEditor/Private/EditMode/SControlRigEditModeTools.cpp
--- a/Engine/Plugins/Animation/ControlRig/Source/ControlRigEditor/Private/EditMode/SControlRigEditModeTools.cpp
+++ b/Engine/Plugins/Animation/ControlRig/Source/ControlRigEditor/Private/EditMode/SControlRigEditModeTools.cpp
@@ -956,7 +956,7 @@ FReply SControlRigEditModeTools::OnBakeControlsToNewSpaceButtonClicked()
 					.Controls(ControlKeys) // use the cached controls here since the selection is not recovered until next tick.
 					.Sequencer(Sequencer)
 					.GetControlCustomization(this, &SControlRigEditModeTools::HandleGetControlElementCustomization)
-					.OnBake_Lambda([Sequencer, ControlRig, TickResolution](URigHierarchy* InHierarchy, TArray<FRigElementKey> InControls, FRigSpacePickerBakeSettings& InSettings)
+					.OnBake_Lambda([Sequencer, ControlRig, TickResolution](URigHierarchy* InHierarchy, TArray<FRigElementKey> InControls, const FRigSpacePickerBakeSettings& InSettings)
 						{		
 							FScopedTransaction Transaction(LOCTEXT("BakeControlToSpace", "Bake Control In Space"));
 							for (const FRigElementKey& ControlKey : InControls)
diff --git a/Engine/Plugins/Editor/ModelingToolsEditorMode/Source/ModelingToolsEditorMode/Private/ModelingToolsEditorModeToolkit.cpp b/Engine/Plugins/Editor/ModelingToolsEditorMode/Source/ModelingToolsEditorMode/Private/ModelingToolsEditorModeToolkit.cpp
--- a/Engine/Plugins/Editor/ModelingToolsEditorMode/Source/ModelingToolsEditorMode/Private/ModelingToolsEditorModeToolkit.cpp
+++ b/Engine/Plugins/Editor/ModelingToolsEditorMode/Source/ModelingToolsEditorMode/Private/ModelingToolsEditorModeToolkit.cpp
@@ -1040,7 +1040,7 @@ TSharedRef<SWidget> FModelingToolsEditorModeToolkit::GetPresetCreateButtonConten
 				[
 					SNew(SEditableTextBox)
 					// Cap the number of characters sent out of the text box, so we don't overflow menus and tooltips
-					.OnTextCommitted_Lambda([this](const FText& NewLabel, ETextCommit::Type&) { NewPresetLabel = NewLabel.ToString().Left(255); })
+					.OnTextCommitted_Lambda([this](const FText& NewLabel, const ETextCommit::Type&) { NewPresetLabel = NewLabel.ToString().Left(255); })
 					.OverflowPolicy(ETextOverflowPolicy::Ellipsis)
 					.ToolTipText(LOCTEXT("ToolPresets_CreatePresetLabel_Tooltip", "A short, descriptive identifier for the new preset."))
 				]
@@ -1059,7 +1059,7 @@ TSharedRef<SWidget> FModelingToolsEditorModeToolkit::GetPresetCreateButtonConten
 				[
 					SNew(SEditableTextBox)
 					// Cap the number of characters sent out of the text box, so we don't overflow menus and tooltips
-					.OnTextCommitted_Lambda([this](const FText& NewToolTip, ETextCommit::Type&) { NewPresetTooltip = NewToolTip.ToString().Left(2048); })
+					.OnTextCommitted_Lambda([this](const FText& NewToolTip, const ETextCommit::Type&) { NewPresetTooltip = NewToolTip.ToString().Left(2048); })
 					.OverflowPolicy(ETextOverflowPolicy::Ellipsis)
 					.ToolTipText(LOCTEXT("ToolPresets_CreatePresetTooltip_Tooltip", "A descriptive tooltip for the new preset."))
 				]
diff --git a/Engine/Source/Developer/AutomationWindow/Private/SAutomationTestItem.cpp b/Engine/Source/Developer/AutomationWindow/Private/SAutomationTestItem.cpp
--- a/Engine/Source/Developer/AutomationWindow/Private/SAutomationTestItem.cpp
+++ b/Engine/Source/Developer/AutomationWindow/Private/SAutomationTestItem.cpp
@@ -537,7 +537,7 @@ FReply SAutomationTestItem::OnEditExcludeOptionsClicked()
 					SNew(SEditableTextBox)
 					.OverflowPolicy(ETextOverflowPolicy::Ellipsis)
 					.Text(FText::FromName(Options->Reason))
-					.OnTextCommitted_Lambda([this, Options](const FText& NewReason, ETextCommit::Type&) { Options->Reason = FName(NewReason.ToString()); })
+					.OnTextCommitted_Lambda([this, Options](const FText& NewReason, const ETextCommit::Type&) { Options->Reason = FName(NewReason.ToString()); })
 					.ToolTipText(LOCTEXT("ExcludeOptions_Reason_ToolTip", "The reason as to why the test is excluded"))
 				]
 			]
diff --git a/Engine/Source/Editor/InternationalizationSettings/Private/InternationalizationSettingsModelDetails.cpp b/Engine/Source/Editor/InternationalizationSettings/Private/InternationalizationSettingsModelDetails.cpp
--- a/Engine/Source/Editor/InternationalizationSettings/Private/InternationalizationSettingsModelDetails.cpp
+++ b/Engine/Source/Editor/InternationalizationSettings/Private/InternationalizationSettingsModelDetails.cpp
@@ -110,7 +110,7 @@ namespace
 		{
 			FCulturePtr EditorLanguage = FInternationalization::Get().GetCurrentLanguage();
 
-			const auto& OnSelectionChangedLambda = [this](FCulturePtr& SelectedCulture, ESelectInfo::Type SelectInfo)
+			const auto& OnSelectionChangedLambda = [this](const FCulturePtr& SelectedCulture, ESelectInfo::Type SelectInfo)
 			{
 				if (SettingsModel.IsValid())
 				{
@@ -220,7 +220,7 @@ namespace
 		{
 			FCulturePtr EditorLocale = FInternationalization::Get().GetCurrentLocale();
 
-			const auto& OnSelectionChangedLambda = [this](FCulturePtr& SelectedCulture, ESelectInfo::Type SelectInfo)
+			const auto& OnSelectionChangedLambda = [this](const FCulturePtr& SelectedCulture, ESelectInfo::Type SelectInfo)
 			{
 				if (SettingsModel.IsValid())
 				{
@@ -331,7 +331,7 @@ namespace
 				}
 			}
 
-			const auto& CulturePickerSelectLambda = [this](FCulturePtr& SelectedCulture, ESelectInfo::Type SelectInfo)
+			const auto& CulturePickerSelectLambda = [this](const FCulturePtr& SelectedCulture, ESelectInfo::Type SelectInfo)
 			{
 				if (SettingsModel.IsValid())
 				{
diff --git a/Engine/Source/Editor/MovieSceneTools/Private/Channels/BuiltInChannelEditors.cpp b/Engine/Source/Editor/MovieSceneTools/Private/Channels/BuiltInChannelEditors.cpp
--- a/Engine/Source/Editor/MovieSceneTools/Private/Channels/BuiltInChannelEditors.cpp
+++ b/Engine/Source/Editor/MovieSceneTools/Private/Channels/BuiltInChannelEditors.cpp
@@ -575,7 +575,7 @@ TSharedRef<SWidget> CreateKeyEditor(const TMovieSceneChannelHandle<FMovieSceneAc
 
 	TSequencerKeyEditor<FMovieSceneActorReferenceData, FMovieSceneActorReferenceKey> KeyEditor(InObjectBindingID, Channel, Section, InSequencer, PropertyBindings, Func);
 
-	auto OnSetCurrentValueLambda = [KeyEditor](FMovieSceneActorReferenceKey& ActorKey) mutable
+	auto OnSetCurrentValueLambda = [KeyEditor](const FMovieSceneActorReferenceKey& ActorKey) mutable
 	{
 		FScopedTransaction Transaction(LOCTEXT("SetActorReferenceKey", "Set Actor Reference Key Value"));
 		KeyEditor.SetValueWithNotify(ActorKey, EMovieSceneDataChangeType::TrackValueChangedRefreshImmediately);
@@ -614,8 +614,10 @@ TSharedRef<SWidget> CreateKeyEditor(const TMovieSceneChannelHandle<FMovieSceneAc
 		TSharedPtr<SWidget> ComponentMenuWidget =
 			SNew(SComponentChooserPopup)
 			.Actor(Actor)
-			.OnComponentChosen_Lambda([=](FName InComponentName) mutable
+			.OnComponentChosen_Lambda([Actor, LevelEditor, KeyEditor, ActorKey = ActorKey](FName InComponentName) mutable
 				{
+					// ActorKey is self-captured so that the lambda can mutate its copy.
+
 					ActorKey.ComponentName = InComponentName;
 					KeyEditor.SetValueWithNotify(ActorKey, EMovieSceneDataChangeType::TrackValueChangedRefreshImmediately);
 
