ModelingMode: Add missing icons to several new tools.

diff --git a/Engine/Plugins/Editor/ModelingToolsEditorMode/Source/ModelingToolsEditorMode/Private/ModelingToolsEditorModeStyle.cpp b/Engine/Plugins/Editor/ModelingToolsEditorMode/Source/ModelingToolsEditorMode/Private/ModelingToolsEditorModeStyle.cpp
--- a/Engine/Plugins/Editor/ModelingToolsEditorMode/Source/ModelingToolsEditorMode/Private/ModelingToolsEditorModeStyle.cpp
+++ b/Engine/Plugins/Editor/ModelingToolsEditorMode/Source/ModelingToolsEditorMode/Private/ModelingToolsEditorModeStyle.cpp
@@ -157,6 +157,8 @@ void FModelingToolsEditorModeStyle::Initialize()
 		StyleSet->Set("ModelingToolsManagerCommands.BeginAddSpherePrimitiveTool.Small", 		new IMAGE_BRUSH_SVG("Icons/ModelingSphere",		Icon40x40));
 		StyleSet->Set("ModelingToolsManagerCommands.BeginAddStairsPrimitiveTool", 			new IMAGE_BRUSH_SVG("Icons/Staircase", Icon20x20));
 		StyleSet->Set("ModelingToolsManagerCommands.BeginAddStairsPrimitiveTool.Small",		new IMAGE_BRUSH_SVG("Icons/Staircase", Icon40x40));
+		StyleSet->Set("ModelingToolsManagerCommands.BeginDrawSplineTool",	     	        new IMAGE_BRUSH_SVG("Icons/GeometryDrawSpline", Icon20x20));
+		StyleSet->Set("ModelingToolsManagerCommands.BeginDrawSplineTool.Small",		        new IMAGE_BRUSH_SVG("Icons/GeometryDrawSpline", Icon40x40));
 
 		StyleSet->Set("ModelingToolsManagerCommands.BeginDrawPolygonTool", 				new IMAGE_PLUGIN_BRUSH("Icons/DrawPolygon_40x",		Icon20x20));
 		StyleSet->Set("ModelingToolsManagerCommands.BeginDrawPolygonTool.Small", 		new IMAGE_PLUGIN_BRUSH("Icons/DrawPolygon_40x", 	Icon20x20));
@@ -313,6 +315,9 @@ void FModelingToolsEditorModeStyle::Initialize()
 		StyleSet->Set("ModelingToolsManagerCommands.BeginSplitMeshesTool.Small", new IMAGE_BRUSH_SVG("Icons/GeometrySplit", Icon20x20));
 		StyleSet->Set("ModelingToolsManagerCommands.BeginPatternTool", new IMAGE_BRUSH_SVG("Icons/ModelingPattern", Icon20x20));
 		StyleSet->Set("ModelingToolsManagerCommands.BeginPatternTool.Small", new IMAGE_BRUSH_SVG("Icons/ModelingPattern", Icon20x20));
+		StyleSet->Set("ModelingToolsManagerCommands.BeginMeshVertexPaintTool", new IMAGE_BRUSH_SVG("Icons/PaintVertexColors", Icon20x20));
+		StyleSet->Set("ModelingToolsManagerCommands.BeginMeshVertexPaintTool.Small", new IMAGE_BRUSH_SVG("Icons/PaintVertexColors", Icon20x20));
+
 
 		StyleSet->Set("ModelingToolsManagerCommands.BeginVolumeToMeshTool",                  new IMAGE_PLUGIN_BRUSH("Icons/ModelingVol2Mesh_x40",         Icon20x20));
 		StyleSet->Set("ModelingToolsManagerCommands.BeginVolumeToMeshTool.Small",            new IMAGE_PLUGIN_BRUSH("Icons/ModelingVol2Mesh_x40",         Icon20x20));
@@ -320,12 +325,14 @@ void FModelingToolsEditorModeStyle::Initialize()
 		StyleSet->Set("ModelingToolsManagerCommands.BeginMeshToVolumeTool.Small",            new IMAGE_PLUGIN_BRUSH("Icons/ModelingMesh2Vol_x40",         Icon20x20));
 		StyleSet->Set("ModelingToolsManagerCommands.BeginBspConversionTool",                 new IMAGE_PLUGIN_BRUSH("Icons/ModelingBSPConversion_x40",         Icon20x20));
 		StyleSet->Set("ModelingToolsManagerCommands.BeginBspConversionTool.Small",           new IMAGE_PLUGIN_BRUSH("Icons/ModelingBSPConversion_x40",         Icon20x20));
-		StyleSet->Set("ModelingToolsManagerCommands.BeginPhysicsInspectorTool",              new IMAGE_PLUGIN_BRUSH("Icons/ModelingPhysInspect_x40",         Icon20x20));
-		StyleSet->Set("ModelingToolsManagerCommands.BeginPhysicsInspectorTool.Small",        new IMAGE_PLUGIN_BRUSH("Icons/ModelingPhysInspect_x40",         Icon20x20));
+		StyleSet->Set("ModelingToolsManagerCommands.BeginPhysicsInspectorTool",              new IMAGE_BRUSH_SVG("Icons/InspectCollision",                     Icon20x20));
+		StyleSet->Set("ModelingToolsManagerCommands.BeginPhysicsInspectorTool.Small",        new IMAGE_BRUSH_SVG("Icons/InspectCollision",                     Icon20x20));
 		StyleSet->Set("ModelingToolsManagerCommands.BeginSetCollisionGeometryTool",          new IMAGE_PLUGIN_BRUSH("Icons/ModelingMeshToCollision_x40",         Icon20x20));
 		StyleSet->Set("ModelingToolsManagerCommands.BeginSetCollisionGeometryTool.Small",    new IMAGE_PLUGIN_BRUSH("Icons/ModelingMeshToCollision_x40",         Icon20x20));
 		StyleSet->Set("ModelingToolsManagerCommands.BeginExtractCollisionGeometryTool",      new IMAGE_PLUGIN_BRUSH("Icons/ModelingCollisionToMesh_x40",         Icon20x20));
 		StyleSet->Set("ModelingToolsManagerCommands.BeginExtractCollisionGeometryTool.Small",new IMAGE_PLUGIN_BRUSH("Icons/ModelingCollisionToMesh_x40",         Icon20x20));
+		StyleSet->Set("ModelingToolsManagerCommands.BeginSimpleCollisionEditorTool",         new IMAGE_BRUSH_SVG("Icons/SimpleCollisionEditor",                  Icon20x20));
+		StyleSet->Set("ModelingToolsManagerCommands.BeginSimpleCollisionEditorTool.Small",   new IMAGE_BRUSH_SVG("Icons/SimpleCollisionEditor",                  Icon20x20));
 
 		StyleSet->Set("ModelingToolsManagerCommands.BeginGenerateStaticMeshLODAssetTool", new IMAGE_BRUSH_SVG("Icons/AutoLOD", Icon20x20));
 		StyleSet->Set("ModelingToolsManagerCommands.BeginGenerateStaticMeshLODAssetTool.Small", new IMAGE_BRUSH_SVG("Icons/AutoLOD", Icon20x20));
