ModelingMode: selection system bugfixes & improvements.

Add DynamicMeshSelector::UpdateAfterGeometryEdit API. StaticMeshSelector implementation updates static mesh after an Edit instead of emitting MeshChange on temporary DynamicMesh.
Delete and Retriangulate Commands now use this API instead of directly emitting transaction, so now these commands work properly on Static Meshes.

FStaticMeshComponentSelectorFactory::CanBuildForTarget now only allows UStaticMeshComponent specifically, filtering out subclasses. This is not ideal but I don't know what else we can do for now, there are many subclasses like ISMC, SplineMeshComponent, etc, that will not work w/ the Selection system. Also now filtering out Engine assets and cooked assets.

ModelingToolsEditorMode now listens for blueprint pre-compiles, and when this occurs, clears the active selection and selection targets. This is necessary because if the selection Component was part of a BP, on recompile it is re-instanced and the old pointer goes stale. Possibly can handle this better or at a more granular level, but clearing the selection is safest. This currently results in things not being undoable because the FChange transactions are on the "old" UDynamicMesh that no longer exists.

UModelingToolsEditorMode::UpdateSelectionManagerOnEditorSelectionChange now does a more thorough job of inspecting the current Actor/Component selection to find Components that could work w/ the selection system.

diff --git a/Engine/Plugins/Editor/ModelingToolsEditorMode/Source/ModelingToolsEditorMode/Private/ModelingToolsEditorMode.cpp b/Engine/Plugins/Editor/ModelingToolsEditorMode/Source/ModelingToolsEditorMode/Private/ModelingToolsEditorMode.cpp
--- a/Engine/Plugins/Editor/ModelingToolsEditorMode/Source/ModelingToolsEditorMode/Private/ModelingToolsEditorMode.cpp
+++ b/Engine/Plugins/Editor/ModelingToolsEditorMode/Source/ModelingToolsEditorMode/Private/ModelingToolsEditorMode.cpp
@@ -1058,6 +1058,8 @@ void UModelingToolsEditorMode::Enter()
 		});
 	}
 
+	BlueprintPreCompileHandle = GEditor->OnBlueprintPreCompile().AddUObject(this, &UModelingToolsEditorMode::OnBlueprintPreCompile);
+
 	// Removing levels from the world can happen either by entering/exiting level instance edit mode, or
 	// by using the Levels panel. The problem is that any temporary actors we may have spawned in the 
 	// level for visualization, gizmos, etc. will be garbage collected. While EdModeInteractiveToolsContext
@@ -1146,6 +1148,10 @@ void UModelingToolsEditorMode::RegisterUVEditor()
 void UModelingToolsEditorMode::Exit()
 {
 	FWorldDelegates::PreLevelRemovedFromWorld.RemoveAll(this);
+	if (BlueprintPreCompileHandle.IsValid())
+	{
+		GEditor->OnBlueprintPreCompile().Remove(BlueprintPreCompileHandle);
+	}
 
 	// shutdown selection interaction
 	if (SelectionInteraction != nullptr)
@@ -1264,6 +1270,23 @@ void UModelingToolsEditorMode::OnEditorClosed()
 	}
 }
 
+
+void UModelingToolsEditorMode::OnBlueprintPreCompile(UBlueprint* Blueprint)
+{
+	// if a Blueprint is compiled, all old instances of it in a level go "stale" and new
+	// instances are created. Currently SelectionManager does not handle this replacement. 
+	// Seems quite hard to know if Blueprint is a parent of any active targets, so if
+	// a Blueprint is compiled we will just clear out any active selection & targets to avoid
+	// potential crashes. Note that this also breaks undo somewhat, as the FChanges seem to 
+	// still be registered against the 'old' instance pointer and hence are ignored/skipped.
+	if (SelectionManager != nullptr)
+	{
+		SelectionManager->ClearSelection();
+		SelectionManager->ClearActiveTargets();
+	}
+}
+
+
 void UModelingToolsEditorMode::OnToolsContextRender(IToolsContextRenderAPI* RenderAPI)
 {
 	if (SelectionInteraction)
@@ -1387,32 +1410,72 @@ void UModelingToolsEditorMode::UpdateSelectionManagerOnEditorSelectionChange(boo
 		return;
 	}
 
-	// figure out what is selected and then extract valid selection identifiers if there are any
-	TArray<AActor*> SelectedActors;
-	GEditor->GetSelectedActors()->GetSelectedObjects<AActor>(SelectedActors); 
+	// Find selected Component types that are currently supported. Currently determining this via explicit casting,
+	// probably it should be handled by the Selector Factories registered in UModelingToolsEditorMode::Enter(),
+	// possibly via the SelectionManager. 
 
-	// this should probably not be the place we make this kind of decision...but for now it will do
-	TArray<FGeometryIdentifier> ValidIdentifiers;
-	for ( AActor* Actor : SelectedActors)
+	TArray<UDynamicMeshComponent*> SelectedDynamicMeshComponents;
+	TArray<UStaticMeshComponent*> SelectedStaticMeshComponents;
+	TArray<UBrushComponent*> SelectedBrushComponents;
+
+	if (GEditor->GetSelectedComponents()->Num() > 0)
 	{
-		if (Cast<ADynamicMeshActor>(Actor) != nullptr)
+		// if we have supported Components selected on a multi-Component Actor, they will be
+		// returned via these functions
+		GEditor->GetSelectedComponents()->GetSelectedObjects<UDynamicMeshComponent>(SelectedDynamicMeshComponents);
+		GEditor->GetSelectedComponents()->GetSelectedObjects<UStaticMeshComponent>(SelectedStaticMeshComponents);
+		GEditor->GetSelectedComponents()->GetSelectedObjects<UBrushComponent>(SelectedBrushComponents);
+	}
+	// Conceivably this could be an 'else', however currently in the Editor when a Volume Actor is selected,
+	// GetSelectedComponents()->Num() > 0 but no BrushComponent will be found (it appears to be some kind of TypedElement wrapper).
+	// However note that this might result in some unexpected Meshes being Selectable on multi-Component Actors
+	if (SelectedDynamicMeshComponents.Num() == 0 && SelectedStaticMeshComponents.Num() == 0 && SelectedBrushComponents.Num() == 0)
+	{
+		// assume Actor selection, find all valid Components on the selected Actors
+		TArray<AActor*> SelectedActors;
+		GEditor->GetSelectedActors()->GetSelectedObjects<AActor>(SelectedActors);
+		for (AActor* Actor : SelectedActors)
 		{
-			ValidIdentifiers.Add( FGeometryIdentifier::PrimitiveComponent(
-				Cast<ADynamicMeshActor>(Actor)->GetDynamicMeshComponent(), FGeometryIdentifier::EObjectType::DynamicMeshComponent) );
+			Actor->ForEachComponent(false, [&](UActorComponent* Component) 
+			{
+				if (UDynamicMeshComponent* DynamicMeshComponent = Cast<UDynamicMeshComponent>(Component))
+				{
+					SelectedDynamicMeshComponents.Add(DynamicMeshComponent);
+				}
+				if (UStaticMeshComponent* StaticMeshComponent = Cast<UStaticMeshComponent>(Component))
+				{
+					SelectedStaticMeshComponents.Add(StaticMeshComponent);
+				}
+				if (UBrushComponent* BrushComponent = Cast<UBrushComponent>(Component))
+				{
+					SelectedBrushComponents.Add(BrushComponent);
+				}
+			});
 		}
-		else if (bEnableVolumeElementSelection && Cast<AVolume>(Actor) != nullptr)
+	}
+
+	// convert selected Component types into selection Identifiers
+	TArray<FGeometryIdentifier> ValidIdentifiers;
+	for (UDynamicMeshComponent* DynamicMeshComponent : SelectedDynamicMeshComponents)
+	{
+		ValidIdentifiers.Add(FGeometryIdentifier::PrimitiveComponent(DynamicMeshComponent, FGeometryIdentifier::EObjectType::DynamicMeshComponent));
+	}
+	if (bEnableStaticMeshElementSelection)
+	{
+		for (UStaticMeshComponent* StaticMeshComponent : SelectedStaticMeshComponents)
 		{
-			ValidIdentifiers.Add( FGeometryIdentifier::PrimitiveComponent(
-				Cast<AVolume>(Actor)->GetBrushComponent(), FGeometryIdentifier::EObjectType::BrushComponent) );
+			ValidIdentifiers.Add(FGeometryIdentifier::PrimitiveComponent(StaticMeshComponent, FGeometryIdentifier::EObjectType::StaticMeshComponent));
 		}
-		else if (bEnableStaticMeshElementSelection && Cast<AStaticMeshActor>(Actor) != nullptr)
+	}
+	if (bEnableVolumeElementSelection)
+	{
+		for (UBrushComponent* BrushComponent : SelectedBrushComponents)
 		{
-			// StaticMeshComponent can exist under any Actor, should not rely on Actor selection...
-			ValidIdentifiers.Add( FGeometryIdentifier::PrimitiveComponent(
-				Cast<AStaticMeshActor>(Actor)->GetStaticMeshComponent(), FGeometryIdentifier::EObjectType::StaticMeshComponent) );
+			ValidIdentifiers.Add(FGeometryIdentifier::PrimitiveComponent(BrushComponent, FGeometryIdentifier::EObjectType::BrushComponent));
 		}
 	}
 
+
 	// This is gross. If we are entering the Mode, we need to update the SelectionManager w/ the current state. However this
 	// update needs to be undoable. Since we are not part of whatever Transaction was involved in changing modes, we are going
 	// to have to emit our own Transaction, which will then be an explicit undo step the user has to go through :(
diff --git a/Engine/Plugins/Editor/ModelingToolsEditorMode/Source/ModelingToolsEditorMode/Public/ModelingToolsEditorMode.h b/Engine/Plugins/Editor/ModelingToolsEditorMode/Source/ModelingToolsEditorMode/Public/ModelingToolsEditorMode.h
--- a/Engine/Plugins/Editor/ModelingToolsEditorMode/Source/ModelingToolsEditorMode/Public/ModelingToolsEditorMode.h
+++ b/Engine/Plugins/Editor/ModelingToolsEditorMode/Source/ModelingToolsEditorMode/Public/ModelingToolsEditorMode.h
@@ -21,6 +21,7 @@ class UModelingSceneSnappingManager;
 class UModelingSelectionInteraction;
 class UGeometrySelectionManager;
 class UInteractiveCommand;
+class UBlueprint;
 
 UCLASS(Transient)
 class UModelingToolsEditorMode : public UBaseLegacyWidgetEdMode, public ILegacyEdModeSelectInterface
@@ -149,6 +150,9 @@ protected:
 
 	void ConfigureRealTimeViewportsOverride(bool bEnable);
 
+	FDelegateHandle BlueprintPreCompileHandle;
+	void OnBlueprintPreCompile(UBlueprint* Blueprint);
+
 
 	// UInteractiveCommand support. Currently implemented by creating instances of
 	// commands on mode startup and holding onto them. This perhaps should be revisited,
diff --git a/Engine/Plugins/Runtime/MeshModelingToolset/Source/MeshModelingTools/Private/Commands/DeleteGeometrySelectionCommand.cpp b/Engine/Plugins/Runtime/MeshModelingToolset/Source/MeshModelingTools/Private/Commands/DeleteGeometrySelectionCommand.cpp
--- a/Engine/Plugins/Runtime/MeshModelingToolset/Source/MeshModelingTools/Private/Commands/DeleteGeometrySelectionCommand.cpp
+++ b/Engine/Plugins/Runtime/MeshModelingToolset/Source/MeshModelingTools/Private/Commands/DeleteGeometrySelectionCommand.cpp
@@ -8,6 +8,7 @@
 #include "DynamicMeshEditor.h"
 #include "Changes/MeshChange.h"
 #include "Selections/GeometrySelectionUtil.h"
+#include "Selection/DynamicMeshSelector.h"
 
 #include UE_INLINE_GENERATED_CPP_BY_NAME(DeleteGeometrySelectionCommand)
 
@@ -31,6 +32,13 @@ bool UDeleteGeometrySelectionCommand::CanExecuteCommandForSelection(UGeometrySel
 
 void UDeleteGeometrySelectionCommand::ExecuteCommandForSelection(UGeometrySelectionEditCommandArguments* SelectionArgs, UInteractiveCommandResult** Result)
 {
+	IGeometrySelector* BaseSelector = SelectionArgs->SelectionHandle.Selector;
+	if (!ensure(BaseSelector != nullptr))
+	{
+		UE_LOG(LogGeometry, Warning, TEXT("UDeleteGeometrySelectionCommand: Delete Selection requires Selector be provided in Selection Arguments"));
+		return;
+	}
+
 	// delete never returns a new selection
 	if (Result != nullptr)
 	{
@@ -39,7 +47,12 @@ void UDeleteGeometrySelectionCommand::ExecuteCommandForSelection(UGeometrySelect
 
 	// should have been verified by CanExecute
 	check(SelectionArgs->IsMatchingType(FGeometryIdentifier::ETargetType::MeshContainer, FGeometryIdentifier::EObjectType::DynamicMesh));
-	
+
+	// TODO: extremely hardcoded behavior right here. Need a way to make this more generic, however
+	// having the UpdateAfterGeometryEdit function in the base GeometrySelector does not make sense as 
+	// it is specific to meshes. Probably this Command needs to be specialized for Mesh Edits.
+	FBaseDynamicMeshSelector* BaseDynamicMeshSelector = static_cast<FBaseDynamicMeshSelector*>(BaseSelector);
+
 	// collect up all our inputs
 	UDynamicMesh* MeshObject = SelectionArgs->SelectionHandle.Identifier.GetAsObjectType<UDynamicMesh>();
 	check(MeshObject != nullptr);
@@ -86,8 +99,10 @@ void UDeleteGeometrySelectionCommand::ExecuteCommandForSelection(UGeometrySelect
 	if ( bTrackChanges && DynamicMeshChange.IsValid() )
 	{
 		SelectionArgs->GetTransactionsAPI()->BeginUndoTransaction(GetCommandShortString());
-		SelectionArgs->GetTransactionsAPI()->AppendChange(MeshObject, 
-			MakeUnique<FMeshChange>(MoveTemp(DynamicMeshChange)), GetCommandShortString());
+
+		BaseDynamicMeshSelector->UpdateAfterGeometryEdit(
+			SelectionArgs->GetTransactionsAPI(), true, MoveTemp(DynamicMeshChange), GetCommandShortString());
+
 		SelectionArgs->GetTransactionsAPI()->EndUndoTransaction();
 	}
 
diff --git a/Engine/Plugins/Runtime/MeshModelingToolset/Source/MeshModelingTools/Private/Commands/RetriangulateGeometrySelectionCommand.cpp b/Engine/Plugins/Runtime/MeshModelingToolset/Source/MeshModelingTools/Private/Commands/RetriangulateGeometrySelectionCommand.cpp
--- a/Engine/Plugins/Runtime/MeshModelingToolset/Source/MeshModelingTools/Private/Commands/RetriangulateGeometrySelectionCommand.cpp
+++ b/Engine/Plugins/Runtime/MeshModelingToolset/Source/MeshModelingTools/Private/Commands/RetriangulateGeometrySelectionCommand.cpp
@@ -13,6 +13,7 @@
 #include "Operations/SimpleHoleFiller.h"
 #include "MeshRegionBoundaryLoops.h"
 #include "DynamicMesh/MeshIndexUtil.h"
+#include "Selection/DynamicMeshSelector.h"
 
 #include UE_INLINE_GENERATED_CPP_BY_NAME(RetriangulateGeometrySelectionCommand)
 
@@ -38,6 +39,13 @@ bool URetriangulateGeometrySelectionCommand::CanExecuteCommandForSelection(UGeom
 
 void URetriangulateGeometrySelectionCommand::ExecuteCommandForSelection(UGeometrySelectionEditCommandArguments* SelectionArgs, UInteractiveCommandResult** Result)
 {
+	IGeometrySelector* BaseSelector = SelectionArgs->SelectionHandle.Selector;
+	if (!ensure(BaseSelector != nullptr))
+	{
+		UE_LOG(LogGeometry, Warning, TEXT("URetriangulateGeometrySelectionCommand: Retriangulate Selection requires Selector be provided in Selection Arguments"));
+		return;
+	}
+
 	if (Result != nullptr)
 	{
 		*Result = nullptr;
@@ -46,6 +54,11 @@ void URetriangulateGeometrySelectionCommand::ExecuteCommandForSelection(UGeometr
 	// should have been verified by CanExecute
 	check(SelectionArgs->IsMatchingType(FGeometryIdentifier::ETargetType::MeshContainer, FGeometryIdentifier::EObjectType::DynamicMesh));
 	
+	// TODO: extremely hardcoded behavior right here. Need a way to make this more generic, however
+	// having the UpdateAfterGeometryEdit function in the base GeometrySelector does not make sense as 
+	// it is specific to meshes. Probably this Command needs to be specialized for Mesh Edits.
+	FBaseDynamicMeshSelector* BaseDynamicMeshSelector = static_cast<FBaseDynamicMeshSelector*>(BaseSelector);
+
 	// collect up all our inputs
 	UDynamicMesh* MeshObject = SelectionArgs->SelectionHandle.Identifier.GetAsObjectType<UDynamicMesh>();
 	check(MeshObject != nullptr);
@@ -170,8 +183,8 @@ void URetriangulateGeometrySelectionCommand::ExecuteCommandForSelection(UGeometr
 	if ( bTrackChanges && DynamicMeshChange.IsValid() )
 	{
 		SelectionArgs->GetTransactionsAPI()->BeginUndoTransaction(GetCommandShortString());
-		SelectionArgs->GetTransactionsAPI()->AppendChange(MeshObject, 
-			MakeUnique<FMeshChange>(MoveTemp(DynamicMeshChange)), GetCommandShortString());
+		BaseDynamicMeshSelector->UpdateAfterGeometryEdit(
+			SelectionArgs->GetTransactionsAPI(), true, MoveTemp(DynamicMeshChange), GetCommandShortString());
 		SelectionArgs->GetTransactionsAPI()->EndUndoTransaction();
 	}
 
diff --git a/Engine/Plugins/Runtime/MeshModelingToolset/Source/ModelingComponents/Private/Selection/DynamicMeshSelector.cpp b/Engine/Plugins/Runtime/MeshModelingToolset/Source/ModelingComponents/Private/Selection/DynamicMeshSelector.cpp
--- a/Engine/Plugins/Runtime/MeshModelingToolset/Source/ModelingComponents/Private/Selection/DynamicMeshSelector.cpp
+++ b/Engine/Plugins/Runtime/MeshModelingToolset/Source/ModelingComponents/Private/Selection/DynamicMeshSelector.cpp
@@ -8,6 +8,7 @@
 #include "DynamicMesh/MeshIndexUtil.h"
 #include "DynamicMesh/MeshNormals.h"
 #include "Changes/MeshVertexChange.h"
+#include "Changes/MeshChange.h"
 #include "DynamicMesh/ColliderMesh.h"
 #include "GroupTopology.h"
 #include "Spatial/SegmentTree3.h"
@@ -628,6 +629,21 @@ void FBaseDynamicMeshSelector::AccumulateSelectionElements(const FGeometrySelect
 }
 
 
+
+
+
+void FBaseDynamicMeshSelector::UpdateAfterGeometryEdit(
+	IToolsContextTransactionsAPI* TransactionsAPI,
+	bool bInTransaction,
+	TUniquePtr<FDynamicMeshChange> DynamicMeshChange,
+	FText GeometryEditTransactionString)
+{
+	TransactionsAPI->AppendChange(GetDynamicMesh(),
+		MakeUnique<FMeshChange>(MoveTemp(DynamicMeshChange)), GeometryEditTransactionString);
+}
+
+
+
 void FBaseDynamicMeshSelector::UpdateColliderMesh()
 {
 	// should we transform to world??
@@ -908,6 +924,7 @@ void FBasicDynamicMeshSelectionTransformer::EndTransform(IToolsContextTransactio
 }
 
 
+
 void FBasicDynamicMeshSelectionTransformer::PreviewRender(IToolsContextRenderAPI* RenderAPI)
 {
 	if (!bEnableSelectionTransformDrawing) return;
diff --git a/Engine/Plugins/Runtime/MeshModelingToolset/Source/ModelingComponents/Private/Selection/GeometrySelectionManager.cpp b/Engine/Plugins/Runtime/MeshModelingToolset/Source/ModelingComponents/Private/Selection/GeometrySelectionManager.cpp
--- a/Engine/Plugins/Runtime/MeshModelingToolset/Source/ModelingComponents/Private/Selection/GeometrySelectionManager.cpp
+++ b/Engine/Plugins/Runtime/MeshModelingToolset/Source/ModelingComponents/Private/Selection/GeometrySelectionManager.cpp
@@ -349,7 +349,10 @@ bool UGeometrySelectionManager::AddActiveTarget(FGeometryIdentifier TargetIdenti
 	}
 
 	TSharedPtr<FGeometrySelectionTarget> SelectionTarget = GetCachedTarget(TargetIdentifier, UseFactory);
-	check(SelectionTarget != nullptr);
+	if (SelectionTarget.IsValid() == false)
+	{
+		return false;
+	}
 
 	ActiveTargetMap.Add(TargetIdentifier, SelectionTarget);
 	ActiveTargetReferences.Add(SelectionTarget);
@@ -536,12 +539,17 @@ TSharedPtr<UGeometrySelectionManager::FGeometrySelectionTarget> UGeometrySelecti
 	}
 
 	// if we are in a situation where we don't have a cache, currently we need the Factory to exist?
-	check(UseFactory != nullptr);
+	if (UseFactory == nullptr)
+	{
+		return nullptr;
+	}
 
 	// selector has to be built properly
 	TUniquePtr<IGeometrySelector> Selector = UseFactory->BuildForTarget(TargetIdentifier);
-	// not going to handle this for now...
-	check(Selector.IsValid());
+	if (Selector.IsValid() == false)
+	{
+		return nullptr;
+	}
 
 	TSharedPtr<FGeometrySelectionTarget> SelectionTarget = MakeShared<FGeometrySelectionTarget>();
 	SelectionTarget->Selector = MoveTemp(Selector);
diff --git a/Engine/Plugins/Runtime/MeshModelingToolset/Source/ModelingComponents/Public/Selection/DynamicMeshSelector.h b/Engine/Plugins/Runtime/MeshModelingToolset/Source/ModelingComponents/Public/Selection/DynamicMeshSelector.h
--- a/Engine/Plugins/Runtime/MeshModelingToolset/Source/ModelingComponents/Public/Selection/DynamicMeshSelector.h
+++ b/Engine/Plugins/Runtime/MeshModelingToolset/Source/ModelingComponents/Public/Selection/DynamicMeshSelector.h
@@ -9,6 +9,7 @@
 #include "Selection/GeometrySelector.h"
 #include "Components/DynamicMeshComponent.h"
 #include "DynamicMesh/DynamicMeshAABBTree3.h"
+#include "DynamicMesh/DynamicMeshChangeTracker.h"
 #include "TransformTypes.h"
 
 
@@ -118,6 +119,27 @@ public:
 	virtual void AccumulateSelectionBounds(const FGeometrySelection& Selection, FGeometrySelectionBounds& BoundsInOut, bool bTransformToWorld) override;
 	virtual void AccumulateSelectionElements(const FGeometrySelection& Selection, FGeometrySelectionElements& Elements, bool bTransformToWorld, bool bIsForPreview) override;
 
+	/**
+	 * UpdateAfterGeometryEdit should be called after editing the UDynamicMesh owned by the Selector (TargetMesh). 
+	 * This may be an external DynamicMesh in the case of a DynamicMeshComponent, or a temporary UDynamicMesh in
+	 * the case of (eg) the StaticMeshSelector and VolumeSelector subclasses. In DynamicMeshSelector the MeshChange
+	 * can just be emitted as a transaction, this is the default behavior. However in StaticMeshSelector, the StaticMesh 
+	 * needs to be synchronized with the UDynamicMesh modification in the same transaction. And potentially the MeshChange 
+	 * does not need to be emitted at all in that case. 
+	 * 
+	 * This is a bit ugly and might be possible to do more cleanly by having the StaticMeshSelector listen to the 
+	 * UDynamicMesh for changes. However currently we do not have the granularity to have it /only/ listen for 
+	 * external mesh edit changes, and not all changes (and since the UDynamicMesh changes in response to  
+	 * StaticMesh changes, eg on undo or external edits, it creates a cycle). 
+	 */
+	virtual void UpdateAfterGeometryEdit(
+		IToolsContextTransactionsAPI* TransactionsAPI, 
+		bool bInTransaction, 
+		TUniquePtr<UE::Geometry::FDynamicMeshChange> DynamicMeshChange, 
+		FText GeometryEditTransactionString );
+
+
+
 public:
 	// these need to be public for the Transformers...can we do it another way?
 	UDynamicMesh* GetDynamicMesh() const { return TargetMesh.Get(); }
diff --git a/Engine/Plugins/Runtime/MeshModelingToolset/Source/ModelingComponentsEditorOnly/Private/Selection/StaticMeshSelector.cpp b/Engine/Plugins/Runtime/MeshModelingToolset/Source/ModelingComponentsEditorOnly/Private/Selection/StaticMeshSelector.cpp
--- a/Engine/Plugins/Runtime/MeshModelingToolset/Source/ModelingComponentsEditorOnly/Private/Selection/StaticMeshSelector.cpp
+++ b/Engine/Plugins/Runtime/MeshModelingToolset/Source/ModelingComponentsEditorOnly/Private/Selection/StaticMeshSelector.cpp
@@ -19,6 +19,7 @@
 #include "Selection/DynamicMeshPolygroupTransformer.h"
 
 #include "RenderingThread.h"
+#include "UObject/Package.h"
 
 using namespace UE::Geometry;
 
@@ -180,6 +181,26 @@ void FStaticMeshSelector::ShutdownTransformation(IGeometrySelectionTransformer*
 }
 
 
+void FStaticMeshSelector::UpdateAfterGeometryEdit(
+	IToolsContextTransactionsAPI* TransactionsAPI,
+	bool bInTransaction,
+	TUniquePtr<FDynamicMeshChange> DynamicMeshChange,
+	FText GeometryEditTransactionString)
+{
+	if (!bInTransaction)
+	{
+		TransactionsAPI->BeginUndoTransaction(GeometryEditTransactionString);
+	}
+
+	CommitMeshTransform();
+
+	if (!bInTransaction)
+	{
+		TransactionsAPI->EndUndoTransaction();
+	}
+}
+
+
 void FStaticMeshSelector::CopyFromStaticMesh()
 {
 	int32 UseLODIndex = 0;
@@ -240,8 +261,29 @@ bool FStaticMeshComponentSelectorFactory::CanBuildForTarget(FGeometryIdentifier
 	if (TargetIdentifier.TargetType == FGeometryIdentifier::ETargetType::PrimitiveComponent)
 	{
 		UStaticMeshComponent* Component = TargetIdentifier.GetAsComponentType<UStaticMeshComponent>();
-		if (Component != nullptr)
+		if (Component != nullptr )
 		{
+			// Ensure that the Component is an SMComponent and not an ISMC or other subclass that might have other behaviors
+			// that we can't know if the Selector will properly support
+			if (ExactCast<UStaticMeshComponent>(Component) == nullptr)
+			{
+				return false;
+			}
+
+			// ensure that we have a static mesh and it's not a built-in Engine mesh
+			UStaticMesh* StaticMesh = Component->GetStaticMesh();
+			if (StaticMesh == nullptr ||
+				StaticMesh->GetPathName().StartsWith(TEXT("/Engine/")))
+			{
+				return false;
+			}
+
+			// ensure that this is not a cooked static mesh
+			if (StaticMesh->GetOutermost()->bIsCookedForEditor)
+			{
+				return false;
+			}
+
 			return true;
 		}
 	}
diff --git a/Engine/Plugins/Runtime/MeshModelingToolset/Source/ModelingComponentsEditorOnly/Private/Selection/VolumeSelector.cpp b/Engine/Plugins/Runtime/MeshModelingToolset/Source/ModelingComponentsEditorOnly/Private/Selection/VolumeSelector.cpp
--- a/Engine/Plugins/Runtime/MeshModelingToolset/Source/ModelingComponentsEditorOnly/Private/Selection/VolumeSelector.cpp
+++ b/Engine/Plugins/Runtime/MeshModelingToolset/Source/ModelingComponentsEditorOnly/Private/Selection/VolumeSelector.cpp
@@ -200,6 +200,29 @@ void FVolumeSelector::ShutdownTransformation(IGeometrySelectionTransformer* Tran
 }
 
 
+
+void FVolumeSelector::UpdateAfterGeometryEdit(
+	IToolsContextTransactionsAPI* TransactionsAPI,
+	bool bInTransaction,
+	TUniquePtr<FDynamicMeshChange> DynamicMeshChange,
+	FText GeometryEditTransactionString)
+{
+	if (!bInTransaction)
+	{
+		TransactionsAPI->BeginUndoTransaction(GeometryEditTransactionString);
+	}
+
+	CommitMeshTransform();
+
+	if (!bInTransaction)
+	{
+		TransactionsAPI->EndUndoTransaction();
+	}
+}
+
+
+
+
 void FVolumeSelector::CommitMeshTransform()
 {
 	FTransform SetTransform = ParentVolume->GetActorTransform();
diff --git a/Engine/Plugins/Runtime/MeshModelingToolset/Source/ModelingComponentsEditorOnly/Public/Selection/StaticMeshSelector.h b/Engine/Plugins/Runtime/MeshModelingToolset/Source/ModelingComponentsEditorOnly/Public/Selection/StaticMeshSelector.h
--- a/Engine/Plugins/Runtime/MeshModelingToolset/Source/ModelingComponentsEditorOnly/Public/Selection/StaticMeshSelector.h
+++ b/Engine/Plugins/Runtime/MeshModelingToolset/Source/ModelingComponentsEditorOnly/Public/Selection/StaticMeshSelector.h
@@ -41,6 +41,12 @@ public:
 	virtual IGeometrySelectionTransformer* InitializeTransformation(const FGeometrySelection& Selection) override;
 	virtual void ShutdownTransformation(IGeometrySelectionTransformer* Transformer) override;
 
+	virtual void UpdateAfterGeometryEdit(
+		IToolsContextTransactionsAPI* TransactionsAPI,
+		bool bInTransaction,
+		TUniquePtr<UE::Geometry::FDynamicMeshChange> DynamicMeshChange,
+		FText GeometryEditTransactionString) override;
+
 protected:
 	UStaticMeshComponent* StaticMeshComponent = nullptr;
 	UStaticMesh* StaticMesh = nullptr;
diff --git a/Engine/Plugins/Runtime/MeshModelingToolset/Source/ModelingComponentsEditorOnly/Public/Selection/VolumeSelector.h b/Engine/Plugins/Runtime/MeshModelingToolset/Source/ModelingComponentsEditorOnly/Public/Selection/VolumeSelector.h
--- a/Engine/Plugins/Runtime/MeshModelingToolset/Source/ModelingComponentsEditorOnly/Public/Selection/VolumeSelector.h
+++ b/Engine/Plugins/Runtime/MeshModelingToolset/Source/ModelingComponentsEditorOnly/Public/Selection/VolumeSelector.h
@@ -39,6 +39,12 @@ public:
 	virtual IGeometrySelectionTransformer* InitializeTransformation(const FGeometrySelection& Selection) override;
 	virtual void ShutdownTransformation(IGeometrySelectionTransformer* Transformer) override;
 
+	virtual void UpdateAfterGeometryEdit(
+		IToolsContextTransactionsAPI* TransactionsAPI,
+		bool bInTransaction,
+		TUniquePtr<UE::Geometry::FDynamicMeshChange> DynamicMeshChange,
+		FText GeometryEditTransactionString) override;
+
 
 	// FEditorUndoClient implementation
 	virtual void PostUndo(bool bSuccess) override;
