Skeleleton Editing Tools icons

diff --git a/Engine/Plugins/Editor/ModelingToolsEditorMode/Source/ModelingToolsEditorMode/Private/ModelingToolsEditorModeStyle.cpp b/Engine/Plugins/Editor/ModelingToolsEditorMode/Source/ModelingToolsEditorMode/Private/ModelingToolsEditorModeStyle.cpp
--- a/Engine/Plugins/Editor/ModelingToolsEditorMode/Source/ModelingToolsEditorMode/Private/ModelingToolsEditorModeStyle.cpp
+++ b/Engine/Plugins/Editor/ModelingToolsEditorMode/Source/ModelingToolsEditorMode/Private/ModelingToolsEditorModeStyle.cpp
@@ -123,8 +123,8 @@ void FModelingToolsEditorModeStyle::Initialize()
 		StyleSet->Set("ModelingToolsManagerCommands.LoadVolumeTools", new IMAGE_BRUSH_SVG( "Icons/LoadVolumeTools", Icon20x20 ) );
 		StyleSet->Set("ModelingToolsManagerCommands.LoadLodsTools", new IMAGE_BRUSH_SVG( "Icons/LoadLodsTools", Icon20x20 ) );
 		
-		StyleSet->Set("ModelingToolsManagerCommands.LoadSkinTools", new IMAGE_BRUSH_SVG( "Icons/LoadDeformTools", Icon20x20 ) );
-		StyleSet->Set("ModelingToolsManagerCommands.LoadSkeletonTools", new IMAGE_BRUSH_SVG( "Icons/LoadTransformTools", Icon20x20 ) );
+		StyleSet->Set("ModelingToolsManagerCommands.LoadSkinTools", new IMAGE_BRUSH_SVG( "Icons/Skin", Icon20x20 ) );
+		StyleSet->Set("ModelingToolsManagerCommands.LoadSkeletonTools", new IMAGE_BRUSH_SVG( "Icons/SkeletalEditor_20", Icon20x20 ) );
 
 		StyleSet->Set("ModelingToolsManagerCommands.BeginShapeSprayTool", 				new IMAGE_PLUGIN_BRUSH("Icons/ShapeSpray_40x",	Icon20x20));
 		StyleSet->Set("ModelingToolsManagerCommands.BeginShapeSprayTool.Small", 		new IMAGE_PLUGIN_BRUSH("Icons/ShapeSpray_40x",	Icon20x20));
@@ -346,10 +346,10 @@ void FModelingToolsEditorModeStyle::Initialize()
 		StyleSet->Set("ModelingToolsManagerCommands.BeginGroomToMeshTool", new IMAGE_BRUSH_SVG("Icons/HairHelmet", Icon20x20));
 		StyleSet->Set("ModelingToolsManagerCommands.BeginGroomToMeshTool.Small", new IMAGE_BRUSH_SVG("Icons/HairHelmet", Icon20x20));
 		
-		StyleSet->Set("ModelingToolsManagerCommands.BeginSkeletonEditingTool", new IMAGE_BRUSH_SVG("Icons/EditSkeleton", Icon20x20));
+		StyleSet->Set("ModelingToolsManagerCommands.BeginSkeletonEditingTool", new IMAGE_BRUSH_SVG("Icons/SkeletalEditor_20", Icon20x20));
 		StyleSet->Set("ModelingToolsManagerCommands.BeginSkinWeightsBindingTool", new IMAGE_BRUSH_SVG("Icons/BindSkin", Icon20x20));
-		StyleSet->Set("ModelingToolsManagerCommands.BeginSkinWeightsPaintTool", new IMAGE_BRUSH_SVG("Icons/SkinWeightsPaint", Icon20x20));
-		StyleSet->Set("ModelingToolsManagerCommands.BeginSkinWeightsPaintTool.Small", new IMAGE_BRUSH_SVG("Icons/SkinWeightsPaint", Icon20x20));
+		StyleSet->Set("ModelingToolsManagerCommands.BeginSkinWeightsPaintTool", new IMAGE_BRUSH_SVG("Icons/EditWeights", Icon20x20));
+		StyleSet->Set("ModelingToolsManagerCommands.BeginSkinWeightsPaintTool.Small", new IMAGE_BRUSH_SVG("Icons/EditWeights", Icon20x20));
 
 		StyleSet->Set("ModelingToolsManagerCommands.MeshSelectionModeAction_NoSelection", new IMAGE_BRUSH_SVG("Icons/ModSelectionObject_16", Icon20x20));
 		StyleSet->Set("ModelingToolsManagerCommands.MeshSelectionModeAction_MeshTriangles", new IMAGE_BRUSH_SVG("Icons/ModSelectionPolys_16", Icon20x20));
