Mesh Topology Selection Mechanic: enable/disable different selection modes in the details customization

diff --git a/Engine/Plugins/Editor/ModelingToolsEditorMode/Source/ModelingToolsEditorMode/Private/DetailsCustomizations/MeshTopologySelectionMechanicCustomization.cpp b/Engine/Plugins/Editor/ModelingToolsEditorMode/Source/ModelingToolsEditorMode/Private/DetailsCustomizations/MeshTopologySelectionMechanicCustomization.cpp
--- a/Engine/Plugins/Editor/ModelingToolsEditorMode/Source/ModelingToolsEditorMode/Private/DetailsCustomizations/MeshTopologySelectionMechanicCustomization.cpp
+++ b/Engine/Plugins/Editor/ModelingToolsEditorMode/Source/ModelingToolsEditorMode/Private/DetailsCustomizations/MeshTopologySelectionMechanicCustomization.cpp
@@ -68,21 +68,46 @@ void FMeshTopologySelectionMechanicPropertiesDetails::CustomizeDetails(IDetailLa
 	ToolbarBuilder.BeginSection("SelectionFilter");
 	ToolbarBuilder.BeginBlockGroup();
 
-	TSharedPtr<IPropertyHandle> SelectVerticesHandle = DetailBuilder.GetProperty(GET_MEMBER_NAME_CHECKED(UMeshTopologySelectionMechanicProperties, bSelectVertices),
-		UMeshTopologySelectionMechanicProperties::StaticClass());
-	// This first bool we're not going to hide. Instead we're going to end up replacing it with the selection filter
-	//SelectVerticesHandle->MarkHiddenByCustomization();
-	AddToggleButtonForBool(SelectVerticesHandle, "PolyEd.SelectCorners");
+	// This first bool we're not going to hide. Instead we're going to end up replacing it with the selection filter below
+	TSharedPtr<IPropertyHandle> FirstHandle = nullptr;
 
-	TSharedPtr<IPropertyHandle> SelectEdgesHandle = DetailBuilder.GetProperty(GET_MEMBER_NAME_CHECKED(UMeshTopologySelectionMechanicProperties, bSelectEdges),
-		UMeshTopologySelectionMechanicProperties::StaticClass());
-	SelectEdgesHandle->MarkHiddenByCustomization();
-	AddToggleButtonForBool(SelectEdgesHandle, "PolyEd.SelectEdges");
+	// Hide or add a toggle button for the given property based on bCanSelect
+	auto ProcessButtonForBool = [&FirstHandle, AddToggleButtonForBool](TSharedPtr<IPropertyHandle> PropertyHandle, bool bCanSelect, FName IconName, bool& bCurrentlySelecting)
+	{
+		if (bCanSelect)
+		{
+			if (FirstHandle == nullptr)
+			{
+				FirstHandle = PropertyHandle;
+			}
+			else
+			{
+				PropertyHandle->MarkHiddenByCustomization();
+			}
+			AddToggleButtonForBool(PropertyHandle, IconName);
+		}
+		else
+		{
+			bCurrentlySelecting = false;
+			PropertyHandle->MarkHiddenByCustomization();
+		}
+	};
+
+	ProcessButtonForBool(DetailBuilder.GetProperty(GET_MEMBER_NAME_CHECKED(UMeshTopologySelectionMechanicProperties, bSelectVertices), UMeshTopologySelectionMechanicProperties::StaticClass()),
+		Properties->bCanSelectVertices,
+		"PolyEd.SelectCorners",
+		Properties->bSelectVertices);
+
+	ProcessButtonForBool(DetailBuilder.GetProperty(GET_MEMBER_NAME_CHECKED(UMeshTopologySelectionMechanicProperties, bSelectEdges), UMeshTopologySelectionMechanicProperties::StaticClass()),
+		Properties->bCanSelectEdges,
+		"PolyEd.SelectEdges",
+		Properties->bSelectEdges);
+
+	ProcessButtonForBool(DetailBuilder.GetProperty(GET_MEMBER_NAME_CHECKED(UMeshTopologySelectionMechanicProperties, bSelectFaces), UMeshTopologySelectionMechanicProperties::StaticClass()),
+		Properties->bCanSelectFaces,
+		"PolyEd.SelectFaces",
+		Properties->bSelectFaces);
 
-	TSharedPtr<IPropertyHandle> SelectFacesHandle = DetailBuilder.GetProperty(GET_MEMBER_NAME_CHECKED(UMeshTopologySelectionMechanicProperties, bSelectFaces),
-		UMeshTopologySelectionMechanicProperties::StaticClass());
-	SelectFacesHandle->MarkHiddenByCustomization();
-	AddToggleButtonForBool(SelectFacesHandle, "PolyEd.SelectFaces");
 
 	ToolbarBuilder.EndBlockGroup();
 	ToolbarBuilder.EndSection();
@@ -109,23 +134,26 @@ void FMeshTopologySelectionMechanicPropertiesDetails::CustomizeDetails(IDetailLa
 
 	// Now we're going to take one of the old checkboxes and replace it with our seleciton filter. Doing this instead
 	// of creating a custom row allows us to use the original splitter, which will then align with the other splitters.
-	IDetailPropertyRow* SelectionFilterRow = DetailBuilder.EditDefaultProperty(SelectVerticesHandle);
-
-	const FText SelectionFilterLabel = LOCTEXT("SelectionFilterLabel", "Selection Filter");
-
-	SelectionFilterRow->CustomWidget()
-		.NameContent()
-		[
-			SelectVerticesHandle->CreatePropertyNameWidget(SelectionFilterLabel, SelectionFilterLabel)
-		]
-		.ValueContent()
-		[
-			SNew(SBox)
-			.Padding(FMargin(0, ModelingUIConstants::DetailRowVertPadding))
+	if (FirstHandle)
+	{
+		IDetailPropertyRow* SelectionFilterRow = DetailBuilder.EditDefaultProperty(FirstHandle);
+
+		const FText SelectionFilterLabel = LOCTEXT("SelectionFilterLabel", "Selection Filter");
+
+		SelectionFilterRow->CustomWidget()
+			.NameContent()
 			[
-				ToolbarBuilder.MakeWidget()
+				FirstHandle->CreatePropertyNameWidget(SelectionFilterLabel, SelectionFilterLabel)
 			]
-		];
+			.ValueContent()
+			[
+				SNew(SBox)
+				.Padding(FMargin(0, ModelingUIConstants::DetailRowVertPadding))
+				[
+					ToolbarBuilder.MakeWidget()
+				]
+			];
+	}
 }
 
 #undef LOCTEXT_NAMESPACE
\ No newline at end of file
diff --git a/Engine/Plugins/Runtime/MeshModelingToolset/Source/ModelingComponents/Public/Selection/MeshTopologySelectionMechanic.h b/Engine/Plugins/Runtime/MeshModelingToolset/Source/ModelingComponents/Public/Selection/MeshTopologySelectionMechanic.h
--- a/Engine/Plugins/Runtime/MeshModelingToolset/Source/ModelingComponents/Public/Selection/MeshTopologySelectionMechanic.h
+++ b/Engine/Plugins/Runtime/MeshModelingToolset/Source/ModelingComponents/Public/Selection/MeshTopologySelectionMechanic.h
@@ -78,6 +78,11 @@ public:
 	// Used to avoid showing some of the selection filter buttons in triedit (in the detail customization)
 	bool bDisplayPolygroupReliantControls = true;
 
+	// Whether to enable the different selection modes (in the detail customization)
+	bool bCanSelectVertices = true;
+	bool bCanSelectEdges = true;
+	bool bCanSelectFaces = true;
+
 	/** Invert current selection. If selection is empty, has same effect as Select All, and is similarly dependent on selection filter. */
 	UFUNCTION(CallInEditor, Category = SelectionActions)
 	void InvertSelection();
