ModelingMode: Fixes the Vertex Sculpt tool to properly update it's detail panel customizations when falloff types and alpha textures are loaded from presets.

diff --git a/Engine/Plugins/Editor/ModelingToolsEditorMode/Source/ModelingToolsEditorMode/Private/DetailsCustomizations/MeshVertexSculptToolCustomizations.cpp b/Engine/Plugins/Editor/ModelingToolsEditorMode/Source/ModelingToolsEditorMode/Private/DetailsCustomizations/MeshVertexSculptToolCustomizations.cpp
--- a/Engine/Plugins/Editor/ModelingToolsEditorMode/Source/ModelingToolsEditorMode/Private/DetailsCustomizations/MeshVertexSculptToolCustomizations.cpp
+++ b/Engine/Plugins/Editor/ModelingToolsEditorMode/Source/ModelingToolsEditorMode/Private/DetailsCustomizations/MeshVertexSculptToolCustomizations.cpp
@@ -109,12 +109,19 @@ TSharedRef<IDetailCustomization> FVertexBrushSculptPropertiesDetails::MakeInstan
 	return MakeShareable(new FVertexBrushSculptPropertiesDetails);
 }
 
+FVertexBrushSculptPropertiesDetails::~FVertexBrushSculptPropertiesDetails()
+{
+	if (TargetTool.IsValid())
+	{
+		TargetTool.Get()->OnDetailsPanelRequestRebuild.Remove(FalloffTypeUpdateHandle);
+	}
+}
+
 void FVertexBrushSculptPropertiesDetails::CustomizeDetails(IDetailLayoutBuilder& DetailBuilder)
 {
 	TSharedPtr<IPropertyHandle> BrushTypeHandle = DetailBuilder.GetProperty(GET_MEMBER_NAME_CHECKED(UVertexBrushSculptProperties, PrimaryBrushType), UVertexBrushSculptProperties::StaticClass());
 	ensure(BrushTypeHandle->IsValidHandle());
 
-	TArray<TWeakObjectPtr<UObject>> ObjectsBeingCustomized;
 	DetailBuilder.GetObjectsBeingCustomized(ObjectsBeingCustomized);
 	check(ObjectsBeingCustomized.Num() > 0);
 	UVertexBrushSculptProperties* BrushProperties = CastChecked<UVertexBrushSculptProperties>(ObjectsBeingCustomized[0]);
@@ -181,7 +188,7 @@ void FVertexBrushSculptPropertiesDetails::CustomizeDetails(IDetailLayoutBuilder&
 		FalloffTypeItems.Add(NewFalloffTypeItem);
 	}
 
-	TSharedPtr<SComboPanel> FalloffTypeCombo = SNew(SComboPanel)
+	FalloffTypeCombo = SNew(SComboPanel)		
 		.ToolTipText(FalloffTypeHandle->GetToolTipText())
 		.ComboButtonTileSize(FVector2D(18, 18))
 		.FlyoutTileSize(FVector2D(FlyoutIconSize, FlyoutIconSize))
@@ -197,6 +204,13 @@ void FVertexBrushSculptPropertiesDetails::CustomizeDetails(IDetailLayoutBuilder&
 		.FlyoutHeaderText(LOCTEXT("FalloffsHeader", "Falloff Types"))
 		.InitialSelectionIndex(CurrentFalloffTypeIndex);
 
+	FalloffTypeUpdateHandle = TargetTool.Get()->OnDetailsPanelRequestRebuild.AddLambda([this]() {
+			UVertexBrushSculptProperties* BrushProperties = CastChecked<UVertexBrushSculptProperties>(ObjectsBeingCustomized[0]);
+			const TArray<UMeshSculptToolBase::FFalloffTypeInfo>& FalloffTypeInfos = TargetTool->GetRegisteredPrimaryFalloffTypes();
+			int32 CurrentFalloffType = (int32)BrushProperties->PrimaryFalloffType;
+			FalloffTypeCombo->SetSelectionIndex(CurrentFalloffType);
+	});
+
 	//DetailBuilder.HideProperty(FalloffTypeHandle);
 	FalloffTypeHandle->MarkHiddenByCustomization();
 
@@ -437,6 +451,14 @@ TSharedRef<IDetailCustomization> FVertexBrushAlphaPropertiesDetails::MakeInstanc
 	return MakeShareable(new FVertexBrushAlphaPropertiesDetails);
 }
 
+FVertexBrushAlphaPropertiesDetails::~FVertexBrushAlphaPropertiesDetails()
+{
+	if (TargetTool.IsValid())
+	{
+		TargetTool.Get()->OnDetailsPanelRequestRebuild.Remove(AlphaTextureUpdateHandle);
+	}
+}
+
 void FVertexBrushAlphaPropertiesDetails::CustomizeDetails(IDetailLayoutBuilder& DetailBuilder)
 {
 	// TODO: move this to a subsystem or UObject CDO
@@ -451,6 +473,11 @@ void FVertexBrushAlphaPropertiesDetails::CustomizeDetails(IDetailLayoutBuilder&
 	}
 	RecentAlphasProvider = RecentAlphasStatic.RecentAlphas;
 
+	DetailBuilder.GetObjectsBeingCustomized(ObjectsBeingCustomized);
+	check(ObjectsBeingCustomized.Num() > 0);
+	UVertexBrushAlphaProperties* AlphaProperties = CastChecked<UVertexBrushAlphaProperties>(ObjectsBeingCustomized[0]);
+	UMeshVertexSculptTool* Tool = AlphaProperties->Tool.Get();
+	TargetTool = Tool;
 
 
 	TSharedPtr<IPropertyHandle> AlphaHandle = DetailBuilder.GetProperty(GET_MEMBER_NAME_CHECKED(UVertexBrushAlphaProperties, Alpha), UVertexBrushAlphaProperties::StaticClass());
@@ -488,7 +515,7 @@ void FVertexBrushAlphaPropertiesDetails::CustomizeDetails(IDetailLayoutBuilder&
 		BrushAlphasLists.Add(CollectionSet);
 	}
 
-	TSharedPtr<SToolInputAssetComboPanel> AlphaAssetPicker = SNew(SToolInputAssetComboPanel)
+	AlphaAssetPicker = SNew(SToolInputAssetComboPanel)
 		.AssetClassType(UTexture2D::StaticClass())		// can infer from property...
 		.Property(AlphaHandle)
 		.ComboButtonTileSize(FVector2D(ComboIconSize, ComboIconSize))
@@ -497,6 +524,10 @@ void FVertexBrushAlphaPropertiesDetails::CustomizeDetails(IDetailLayoutBuilder&
 		.RecentAssetsProvider(RecentAlphasProvider)
 		.CollectionSets(BrushAlphasLists);
 
+	AlphaTextureUpdateHandle = TargetTool.Get()->OnDetailsPanelRequestRebuild.AddLambda([this]() {
+		AlphaAssetPicker->RefreshThumbnailFromProperty();
+	});
+
 	DetailBuilder.EditDefaultProperty(AlphaHandle)->CustomWidget()
 		.OverrideResetToDefault(FResetToDefaultOverride::Hide())
 		.WholeRowContent()
diff --git a/Engine/Plugins/Editor/ModelingToolsEditorMode/Source/ModelingToolsEditorMode/Public/DetailsCustomizations/MeshVertexSculptToolCustomizations.h b/Engine/Plugins/Editor/ModelingToolsEditorMode/Source/ModelingToolsEditorMode/Public/DetailsCustomizations/MeshVertexSculptToolCustomizations.h
--- a/Engine/Plugins/Editor/ModelingToolsEditorMode/Source/ModelingToolsEditorMode/Public/DetailsCustomizations/MeshVertexSculptToolCustomizations.h
+++ b/Engine/Plugins/Editor/ModelingToolsEditorMode/Source/ModelingToolsEditorMode/Public/DetailsCustomizations/MeshVertexSculptToolCustomizations.h
@@ -14,9 +14,11 @@ class IDetailLayoutBuilder;
 class IPropertyHandle;
 class IDetailChildrenBuilder;
 class SComboButton;
+class SComboPanel;
 class SBox;
 class UMeshVertexSculptTool;
 class FRecentAlphasProvider;
+class SToolInputAssetComboPanel;
 
 
 // customization for USculptBrushProperties, creates two-column layout
@@ -34,6 +36,8 @@ public:
 class FVertexBrushSculptPropertiesDetails : public IDetailCustomization
 {
 public:
+	virtual ~FVertexBrushSculptPropertiesDetails();
+
 	static TSharedRef<IDetailCustomization> MakeInstance();
 	void CustomizeDetails(IDetailLayoutBuilder& DetailBuilder) override;
 
@@ -47,6 +51,10 @@ protected:
 	FReply OnToggledFreezeTarget();
 	void OnSetFreezeTarget(ECheckBoxState State);
 	bool IsFreezeTargetEnabled();
+
+	TSharedPtr<SComboPanel> FalloffTypeCombo;
+	TArray<TWeakObjectPtr<UObject>> ObjectsBeingCustomized;
+	FDelegateHandle FalloffTypeUpdateHandle;
 };
 
 
@@ -56,10 +64,18 @@ protected:
 class FVertexBrushAlphaPropertiesDetails : public IDetailCustomization
 {
 public:
+	virtual ~FVertexBrushAlphaPropertiesDetails();
+
 	static TSharedRef<IDetailCustomization> MakeInstance();
 	void CustomizeDetails(IDetailLayoutBuilder& DetailBuilder) override;
 
 	TSharedPtr<FRecentAlphasProvider> RecentAlphasProvider;
+
+protected:
+	TWeakObjectPtr<UMeshVertexSculptTool> TargetTool;
+	TArray<TWeakObjectPtr<UObject>> ObjectsBeingCustomized;
+	TSharedPtr<SToolInputAssetComboPanel> AlphaAssetPicker;
+	FDelegateHandle AlphaTextureUpdateHandle;
 };
 
 #if UE_ENABLE_INCLUDE_ORDER_DEPRECATED_IN_5_2
diff --git a/Engine/Plugins/Experimental/MeshModelingToolsetExp/Source/MeshModelingToolsExp/Private/MeshVertexSculptTool.cpp b/Engine/Plugins/Experimental/MeshModelingToolsetExp/Source/MeshModelingToolsExp/Private/MeshVertexSculptTool.cpp
--- a/Engine/Plugins/Experimental/MeshModelingToolsetExp/Source/MeshModelingToolsExp/Private/MeshVertexSculptTool.cpp
+++ b/Engine/Plugins/Experimental/MeshModelingToolsetExp/Source/MeshModelingToolsExp/Private/MeshVertexSculptTool.cpp
@@ -204,6 +204,7 @@ void UMeshVertexSculptTool::Setup()
 
 	AlphaProperties = NewObject<UVertexBrushAlphaProperties>(this);
 	AlphaProperties->RestoreProperties(this);
+	AlphaProperties->Tool = this;
 	AddToolPropertySource(AlphaProperties);
 
 	SymmetryProperties = NewObject<UMeshSymmetryProperties>(this);
@@ -300,10 +301,18 @@ void UMeshVertexSculptTool::Setup()
 		[this](EMeshVertexSculptBrushType NewType) { UpdateBrushType(NewType); });
 
 	SculptProperties->WatchProperty( SculptProperties->PrimaryFalloffType,
-		[this](EMeshSculptFalloffType NewType) { SetPrimaryFalloffType(NewType); });
+		[this](EMeshSculptFalloffType NewType) { 
+			SetPrimaryFalloffType(NewType);
+			// Request to have the details panel rebuilt to ensure the new falloff property value is propagated to the details customization
+			OnDetailsPanelRequestRebuild.Broadcast();
+		});
 
 	SculptProperties->WatchProperty(AlphaProperties->Alpha,
-		[this](UTexture2D* NewAlpha) { UpdateBrushAlpha(AlphaProperties->Alpha); });
+		[this](UTexture2D* NewAlpha) {
+			UpdateBrushAlpha(AlphaProperties->Alpha);
+			// Request to have the details panel rebuilt to ensure the new alpha property value is propagated to the details customization
+			OnDetailsPanelRequestRebuild.Broadcast();
+		});
 
 	// must call before updating brush type so that we register all brush properties?
 	UMeshSculptToolBase::OnCompleteSetup();
diff --git a/Engine/Plugins/Experimental/MeshModelingToolsetExp/Source/MeshModelingToolsExp/Public/MeshVertexSculptTool.h b/Engine/Plugins/Experimental/MeshModelingToolsetExp/Source/MeshModelingToolsExp/Public/MeshVertexSculptTool.h
--- a/Engine/Plugins/Experimental/MeshModelingToolsetExp/Source/MeshModelingToolsExp/Public/MeshVertexSculptTool.h
+++ b/Engine/Plugins/Experimental/MeshModelingToolsetExp/Source/MeshModelingToolsExp/Public/MeshVertexSculptTool.h
@@ -168,6 +168,10 @@ public:
 	/** Bounds of random generation (positive and negative) for randomized stamps */
 	UPROPERTY(EditAnywhere, Category = Alpha, AdvancedDisplay, meta = (UIMin = "0.0", UIMax = "180.0"))
 	float RandomRange = 180.0;
+
+	// parent ref required for details customization
+	UPROPERTY(meta = (TransientToolProperty))
+	TWeakObjectPtr<UMeshVertexSculptTool> Tool;
 };
 
 
@@ -235,6 +239,17 @@ public:
 	virtual void SetActiveFalloffType(int32 Identifier);
 	virtual void SetRegionFilterType(int32 Identifier);
 
+	/**
+	* OnDetailsPanelRequestRebuild is broadcast by the tool when it detects it needs to have it's details panel rebuilt outside
+	* of normal rebuilding triggers, such as changing property set objects. This is useful in rare circumstances, such as when
+	* the tool is using detail customizations and tool properties are changed outside of user interactions, such as via tool
+	* preset loading. In these cases, the detail customization widgets might not be updated properly without rebuilding the details
+	* panel completely.
+	*/
+
+	DECLARE_MULTICAST_DELEGATE(OnInteractiveToolDetailsPanelRequestRebuild);
+	OnInteractiveToolDetailsPanelRequestRebuild OnDetailsPanelRequestRebuild;
+
 protected:
 	// UMeshSculptToolBase API
 	virtual void InitializeIndicator() override;
diff --git a/Engine/Plugins/Experimental/MeshModelingToolsetExp/Source/ModelingEditorUI/Private/ModelingWidgets/SComboPanel.cpp b/Engine/Plugins/Experimental/MeshModelingToolsetExp/Source/ModelingEditorUI/Private/ModelingWidgets/SComboPanel.cpp
--- a/Engine/Plugins/Experimental/MeshModelingToolsetExp/Source/ModelingEditorUI/Private/ModelingWidgets/SComboPanel.cpp
+++ b/Engine/Plugins/Experimental/MeshModelingToolsetExp/Source/ModelingEditorUI/Private/ModelingWidgets/SComboPanel.cpp
@@ -269,7 +269,13 @@ void SComboPanel::Construct(const FArguments& InArgs)
 	];
 }
 
+void SComboPanel::SetSelectionIndex(int InSelectionIndex)
+{
+	int SelectionIndex = FMath::Clamp(InSelectionIndex, 0, Items.Num());
 
+	SelectedItem = Items[SelectionIndex];
+	TileView->SetSelection(SelectedItem);
+}
 
 TSharedRef<SWidget> SComboPanel::OnGetMenuContent()
 {
diff --git a/Engine/Plugins/Experimental/MeshModelingToolsetExp/Source/ModelingEditorUI/Private/ModelingWidgets/SToolInputAssetComboPanel.cpp b/Engine/Plugins/Experimental/MeshModelingToolsetExp/Source/ModelingEditorUI/Private/ModelingWidgets/SToolInputAssetComboPanel.cpp
--- a/Engine/Plugins/Experimental/MeshModelingToolsetExp/Source/ModelingEditorUI/Private/ModelingWidgets/SToolInputAssetComboPanel.cpp
+++ b/Engine/Plugins/Experimental/MeshModelingToolsetExp/Source/ModelingEditorUI/Private/ModelingWidgets/SToolInputAssetComboPanel.cpp
@@ -218,14 +218,7 @@ void SToolInputAssetComboPanel::Construct(const FArguments& InArgs)
 
 
 	// set initial thumbnail
-	if (Property.IsValid())
-	{
-		FAssetData AssetData;
-		if (Property->GetValue(AssetData) == FPropertyAccess::Result::Success)
-		{
-			AssetThumbnail->SetAsset(AssetData);
-		}
-	}
+	RefreshThumbnailFromProperty();
 
 	if (InArgs._InitiallySelectedAsset.IsValid())
 	{
@@ -240,6 +233,17 @@ void SToolInputAssetComboPanel::Construct(const FArguments& InArgs)
 }
 
 
+void SToolInputAssetComboPanel::RefreshThumbnailFromProperty()
+{
+	if (Property.IsValid())
+	{
+		FAssetData AssetData;
+		if (Property->GetValue(AssetData) == FPropertyAccess::Result::Success)
+		{
+			AssetThumbnail->SetAsset(AssetData);
+		}
+	}
+}
 
 
 TSharedRef<SWidget> SToolInputAssetComboPanel::MakeCollectionSetsButtonPanel(TSharedRef<SToolInputAssetPicker> AssetPickerView)
diff --git a/Engine/Plugins/Experimental/MeshModelingToolsetExp/Source/ModelingEditorUI/Public/ModelingWidgets/SComboPanel.h b/Engine/Plugins/Experimental/MeshModelingToolsetExp/Source/ModelingEditorUI/Public/ModelingWidgets/SComboPanel.h
--- a/Engine/Plugins/Experimental/MeshModelingToolsetExp/Source/ModelingEditorUI/Public/ModelingWidgets/SComboPanel.h
+++ b/Engine/Plugins/Experimental/MeshModelingToolsetExp/Source/ModelingEditorUI/Public/ModelingWidgets/SComboPanel.h
@@ -93,7 +93,12 @@ public:
 	 */
 	void Construct( const FArguments& InArgs );
 
-
+	/**
+	* Sets the currently selected index without user interaction.
+	* 
+	* @param InSelectionIndex The new selection index for the widget
+	*/
+	void SetSelectionIndex(int InSelectionIndex);
 
 protected:
 	FVector2D ComboButtonTileSize;
diff --git a/Engine/Plugins/Experimental/MeshModelingToolsetExp/Source/ModelingEditorUI/Public/ModelingWidgets/SToolInputAssetComboPanel.h b/Engine/Plugins/Experimental/MeshModelingToolsetExp/Source/ModelingEditorUI/Public/ModelingWidgets/SToolInputAssetComboPanel.h
--- a/Engine/Plugins/Experimental/MeshModelingToolsetExp/Source/ModelingEditorUI/Public/ModelingWidgets/SToolInputAssetComboPanel.h
+++ b/Engine/Plugins/Experimental/MeshModelingToolsetExp/Source/ModelingEditorUI/Public/ModelingWidgets/SToolInputAssetComboPanel.h
@@ -122,6 +122,8 @@ public:
 	 */
 	void Construct( const FArguments& InArgs );
 
+	void RefreshThumbnailFromProperty();
+
 protected:
 
 	FVector2D ComboButtonTileSize;
