ToolPresets: Adjust handling of text display in preset dialogs and menus to prevent egregious widget overflow across the screen.

diff --git a/Engine/Plugins/Editor/ModelingToolsEditorMode/Source/ModelingToolsEditorMode/Private/ModelingToolsEditorModeToolkit.cpp b/Engine/Plugins/Editor/ModelingToolsEditorMode/Source/ModelingToolsEditorMode/Private/ModelingToolsEditorModeToolkit.cpp
--- a/Engine/Plugins/Editor/ModelingToolsEditorMode/Source/ModelingToolsEditorMode/Private/ModelingToolsEditorModeToolkit.cpp
+++ b/Engine/Plugins/Editor/ModelingToolsEditorMode/Source/ModelingToolsEditorMode/Private/ModelingToolsEditorModeToolkit.cpp
@@ -64,6 +64,7 @@
 #include "ModelingWidgets/SToolInputAssetComboPanel.h"
 #include "Fonts/SlateFontInfo.h"
 #include "ToolPresetAssetSubsystem.h"
+#include "Widgets/Input/SMultiLineEditableTextBox.h"
 
 
 #define LOCTEXT_NAMESPACE "FModelingToolsEditorModeToolkit"
@@ -1040,6 +1041,7 @@ TSharedRef<SWidget> FModelingToolsEditorModeToolkit::GetPresetCreateButtonConten
 					.Text(LOCTEXT("ToolPresets_CreatePresetLabel", "Label"))
 				]
 				+ SHorizontalBox::Slot()
+				.MaxWidth(300)
 				[
 					SNew(SEditableTextBox)
 					// Cap the number of characters sent out of the text box, so we don't overflow menus and tooltips
@@ -1059,12 +1061,22 @@ TSharedRef<SWidget> FModelingToolsEditorModeToolkit::GetPresetCreateButtonConten
 					.Text(LOCTEXT("ToolPresets_CreatePresetTooltip", "Tooltip"))
 				]
 				+ SHorizontalBox::Slot()
+				.MaxWidth(300)
 				[
-					SNew(SEditableTextBox)
-					// Cap the number of characters sent out of the text box, so we don't overflow menus and tooltips
-					.OnTextCommitted_Lambda([this](const FText& NewToolTip, const ETextCommit::Type&) { NewPresetTooltip = NewToolTip.ToString().Left(2048); })
-					.OverflowPolicy(ETextOverflowPolicy::Ellipsis)
-					.ToolTipText(LOCTEXT("ToolPresets_CreatePresetTooltip_Tooltip", "A descriptive tooltip for the new preset."))
+					SNew(SBox)
+					.MinDesiredHeight(44.f)
+					.MaxDesiredHeight(44.0f)
+					[
+						SNew(SMultiLineEditableTextBox)
+						// Cap the number of characters sent out of the text box, so we don't overflow menus and tooltips
+						.OnTextCommitted_Lambda([this](const FText& NewToolTip, const ETextCommit::Type&) { NewPresetTooltip = NewToolTip.ToString().Left(2048); })
+						.OverflowPolicy(ETextOverflowPolicy::Ellipsis)
+						.AllowMultiLine(false)
+						.AutoWrapText(true)
+						.WrappingPolicy(ETextWrappingPolicy::DefaultWrapping)
+						.OverflowPolicy(ETextOverflowPolicy::Ellipsis)
+						.ToolTipText(LOCTEXT("ToolPresets_CreatePresetTooltip_Tooltip", "A descriptive tooltip for the new preset."))
+					]
 				]
 			]
 		+ SVerticalBox::Slot()
@@ -1252,8 +1264,22 @@ void FModelingToolsEditorModeToolkit::RebuildPresetListForTool(bool bSettingsOpe
 						continue;
 					}
 					TSharedPtr<FToolPresetOption> NewOption = MakeShared<FToolPresetOption>();
-					NewOption->PresetLabel = PresetDef.Label;
-					NewOption->PresetTooltip = PresetDef.Tooltip;
+					if (PresetDef.Label.Len() > 50)
+					{
+						NewOption->PresetLabel = PresetDef.Label.Left(50) + FString("...");
+					}
+					else
+					{
+						NewOption->PresetLabel = PresetDef.Label;
+					}
+					if (PresetDef.Tooltip.Len() > 2048)
+					{
+						NewOption->PresetTooltip = PresetDef.Tooltip.Left(2048) + FString("...");
+					}
+					else
+					{
+						NewOption->PresetTooltip = PresetDef.Tooltip;
+					}
 					NewOption->PresetIndex = PresetIndex;
 					NewOption->PresetCollection = PresetCollection;
 
