ModelingToolsEditorMode: Updated tool icons.

Added icons for the Revolve Spline, Mesh Splines, and ISM Editor tools. Updated the icons for Revolve Boundary and Revolve Path.

diff --git a/Engine/Plugins/Editor/ModelingToolsEditorMode/Source/ModelingToolsEditorMode/Private/ModelingToolsEditorModeStyle.cpp b/Engine/Plugins/Editor/ModelingToolsEditorMode/Source/ModelingToolsEditorMode/Private/ModelingToolsEditorModeStyle.cpp
--- a/Engine/Plugins/Editor/ModelingToolsEditorMode/Source/ModelingToolsEditorMode/Private/ModelingToolsEditorModeStyle.cpp
+++ b/Engine/Plugins/Editor/ModelingToolsEditorMode/Source/ModelingToolsEditorMode/Private/ModelingToolsEditorModeStyle.cpp
@@ -274,10 +274,16 @@ void FModelingToolsEditorModeStyle::Initialize()
 		StyleSet->Set("ModelingToolsManagerCommands.BeginPolyGroupsTool.Small",              new IMAGE_PLUGIN_BRUSH("Icons/PolyGroups_40x",       Icon20x20));      
 		StyleSet->Set("ModelingToolsManagerCommands.BeginDrawPolyPathTool",                  new IMAGE_PLUGIN_BRUSH("Icons/PolyPath_40x",         Icon20x20));    
 		StyleSet->Set("ModelingToolsManagerCommands.BeginDrawPolyPathTool.Small",            new IMAGE_PLUGIN_BRUSH("Icons/PolyPath_40x",         Icon20x20));    
-		StyleSet->Set("ModelingToolsManagerCommands.BeginDrawAndRevolveTool",                new IMAGE_PLUGIN_BRUSH("Icons/ModelingDrawAndRevolve_x40", Icon20x20));
-		StyleSet->Set("ModelingToolsManagerCommands.BeginDrawAndRevolveTool.Small",          new IMAGE_PLUGIN_BRUSH("Icons/ModelingDrawAndRevolve_x20", Icon20x20));
-		StyleSet->Set("ModelingToolsManagerCommands.BeginRevolveBoundaryTool",               new IMAGE_PLUGIN_BRUSH("Icons/ModelingRevolveBoundary_x40", Icon20x20));
-		StyleSet->Set("ModelingToolsManagerCommands.BeginRevolveBoundaryTool.Small",         new IMAGE_PLUGIN_BRUSH("Icons/ModelingRevolveBoundary_x20", Icon20x20));
+		StyleSet->Set("ModelingToolsManagerCommands.BeginDrawAndRevolveTool",                new IMAGE_BRUSH_SVG("Icons/ModelingDrawAndRevolve",  Icon20x20));
+		StyleSet->Set("ModelingToolsManagerCommands.BeginDrawAndRevolveTool.Small",          new IMAGE_BRUSH_SVG("Icons/ModelingDrawAndRevolve",  Icon20x20));
+		StyleSet->Set("ModelingToolsManagerCommands.BeginRevolveBoundaryTool",               new IMAGE_BRUSH_SVG("Icons/ModelingRevolveBoundary", Icon20x20));
+		StyleSet->Set("ModelingToolsManagerCommands.BeginRevolveBoundaryTool.Small",         new IMAGE_BRUSH_SVG("Icons/ModelingRevolveBoundary", Icon20x20));
+		StyleSet->Set("ModelingToolsManagerCommands.BeginRevolveSplineTool",                 new IMAGE_BRUSH_SVG("Icons/ModelingRevolveSpline",   Icon20x20));
+		StyleSet->Set("ModelingToolsManagerCommands.BeginRevolveSplineTool.Small",           new IMAGE_BRUSH_SVG("Icons/ModelingRevolveSpline",   Icon20x20));
+		StyleSet->Set("ModelingToolsManagerCommands.BeginTriangulateSplinesTool",            new IMAGE_BRUSH_SVG("Icons/ModelingTriangulateSpline",   Icon20x20));
+		StyleSet->Set("ModelingToolsManagerCommands.BeginTriangulateSplinesTool.Small",      new IMAGE_BRUSH_SVG("Icons/ModelingTriangulateSpline",   Icon20x20));
+		StyleSet->Set("ModelingToolsManagerCommands.BeginISMEditorTool",						new IMAGE_BRUSH_SVG("Icons/ModelingISMEditor",       Icon20x20));
+		StyleSet->Set("ModelingToolsManagerCommands.BeginISMEditorTool.Small",				new IMAGE_BRUSH_SVG("Icons/ModelingISMEditor",       Icon20x20));
 		StyleSet->Set("ModelingToolsManagerCommands.BeginCubeGridTool",                      new IMAGE_BRUSH_SVG("Icons/CubeGrid",                Icon20x20));
 		StyleSet->Set("ModelingToolsManagerCommands.BeginCubeGridTool.Small",                new IMAGE_BRUSH_SVG("Icons/CubeGrid",                Icon20x20));
 		StyleSet->Set("ModelingToolsManagerCommands.BeginMeshBooleanTool",                   new IMAGE_PLUGIN_BRUSH("Icons/ModelingMeshBoolean_x40", Icon20x20));
