minor typo fix in modeling mode toolkit toolbar comments

diff --git a/Engine/Plugins/Editor/ModelingToolsEditorMode/Source/ModelingToolsEditorMode/Private/ModelingModeToolkit_Toolbars.cpp b/Engine/Plugins/Editor/ModelingToolsEditorMode/Source/ModelingToolsEditorMode/Private/ModelingModeToolkit_Toolbars.cpp
--- a/Engine/Plugins/Editor/ModelingToolsEditorMode/Source/ModelingToolsEditorMode/Private/ModelingModeToolkit_Toolbars.cpp
+++ b/Engine/Plugins/Editor/ModelingToolsEditorMode/Source/ModelingToolsEditorMode/Private/ModelingModeToolkit_Toolbars.cpp
@@ -50,7 +50,7 @@ void MakeSubMenu_QuickSettings(FMenuBuilder& MenuBuilder)
 			}
 		}), FCanExecuteAction(), FIsActionChecked());
 	MenuBuilder.AddMenuEntry(LOCTEXT("ModelingModeProjectSettings", "Modeling Mode (Project)"), 
-		LOCTEXT("ModelingModeProjectSettings_Tooltip", "Jump to the Project Settings for Modeling Mode. Project Settings are Poject-specific."),
+		LOCTEXT("ModelingModeProjectSettings_Tooltip", "Jump to the Project Settings for Modeling Mode. Project Settings are Project-specific."),
 		FSlateIcon(), OpenModelingModeProjectSettings, NAME_None, EUserInterfaceActionType::Button);
 
 	const FUIAction OpenModelingModeEditorSettings(
@@ -74,7 +74,7 @@ void MakeSubMenu_QuickSettings(FMenuBuilder& MenuBuilder)
 			}
 		}), FCanExecuteAction(), FIsActionChecked());
 	MenuBuilder.AddMenuEntry(LOCTEXT("ModelingToolsProjectSettings", "Modeling Tools (Project)"), 
-		LOCTEXT("ModelingToolsProjectSettings_Tooltip", "Jump to the Project Settings for Modeling Tools. Project Settings are Poject-specific."),
+		LOCTEXT("ModelingToolsProjectSettings_Tooltip", "Jump to the Project Settings for Modeling Tools. Project Settings are Project-specific."),
 		FSlateIcon(), OpenModelingToolsProjectSettings, NAME_None, EUserInterfaceActionType::Button);
 
 }
