Add a method for IModelingModeToolExtensions to provide tool target factories so that plugins can extend the interface with new tooltargets.

diff --git a/Engine/Plugins/Editor/ModelingToolsEditorMode/Source/ModelingToolsEditorMode/Private/ModelingToolsEditorMode.cpp b/Engine/Plugins/Editor/ModelingToolsEditorMode/Source/ModelingToolsEditorMode/Private/ModelingToolsEditorMode.cpp
--- a/Engine/Plugins/Editor/ModelingToolsEditorMode/Source/ModelingToolsEditorMode/Private/ModelingToolsEditorMode.cpp
+++ b/Engine/Plugins/Editor/ModelingToolsEditorMode/Source/ModelingToolsEditorMode/Private/ModelingToolsEditorMode.cpp
@@ -937,6 +937,15 @@ void UModelingToolsEditorMode::Enter()
 
 				RegisterTool(ToolInfo.ToolCommand, ToolInfo.ToolName.ToString(), ToolInfo.ToolBuilder);
 			}
+
+			TArray<TSubclassOf<UToolTargetFactory>> ExtensionToolTargetFactoryClasses;
+			if (Extensions[k]->GetExtensionToolTargets(ExtensionToolTargetFactoryClasses))
+			{
+				for (const TSubclassOf<UToolTargetFactory>& ExtensionTargetFactoryClass : ExtensionToolTargetFactoryClasses)
+				{
+					GetInteractiveToolsContext()->TargetManager->AddTargetFactory(NewObject<UToolTargetFactory>(GetToolManager(), ExtensionTargetFactoryClass.Get()));
+				}
+			}
 		}
 	}
 
diff --git a/Engine/Plugins/Editor/ModelingToolsEditorMode/Source/ModelingToolsEditorMode/Public/ModelingModeToolExtensions.h b/Engine/Plugins/Editor/ModelingToolsEditorMode/Source/ModelingToolsEditorMode/Public/ModelingModeToolExtensions.h
--- a/Engine/Plugins/Editor/ModelingToolsEditorMode/Source/ModelingToolsEditorMode/Public/ModelingModeToolExtensions.h
+++ b/Engine/Plugins/Editor/ModelingToolsEditorMode/Source/ModelingToolsEditorMode/Public/ModelingModeToolExtensions.h
@@ -4,6 +4,7 @@
 
 #include "Features/IModularFeature.h"
 #include "Internationalization/Text.h"
+#include "ToolTargets/ToolTarget.h"
 
 class UInteractiveToolsContext;
 class FUICommandInfo;
@@ -98,6 +99,12 @@ public:
 	 */
 	virtual bool GetExtensionExtendedInfo(FModelingModeExtensionExtendedInfo& InfoOut) { return false; }
 
+	/** 
+	 * Query the extension for additional tool targets. This function is optional and the result will be
+	 * ignored unless it is overridden and returns true.
+	 */
+	virtual bool GetExtensionToolTargets(TArray<TSubclassOf<UToolTargetFactory>>& ToolTargetFactoriesOut) { return false; }
+
 	static FName GetModularFeatureName()
 	{
 		static FName FeatureName = FName(TEXT("ModelingModeToolExtension"));
