MeshModelingToolset: Fix missing LOCTEXT in editor mode asset settings

diff --git a/Engine/Plugins/Editor/ModelingToolsEditorMode/Source/ModelingToolsEditorMode/Private/ModelingToolsEditorModeToolkit.cpp b/Engine/Plugins/Editor/ModelingToolsEditorMode/Source/ModelingToolsEditorMode/Private/ModelingToolsEditorModeToolkit.cpp
--- a/Engine/Plugins/Editor/ModelingToolsEditorMode/Source/ModelingToolsEditorMode/Private/ModelingToolsEditorModeToolkit.cpp
+++ b/Engine/Plugins/Editor/ModelingToolsEditorMode/Source/ModelingToolsEditorMode/Private/ModelingToolsEditorModeToolkit.cpp
@@ -742,15 +742,15 @@ TSharedPtr<SWidget> FModelingToolsEditorModeToolkit::MakeAssetConfigPanel()
 	// New Asset Location drop-down
 	//
 
-	AssetLocationModes.Add(MakeShared<FString>(TEXT("AutoGen Folder (World-Relative)")));
-	AssetLocationModes.Add(MakeShared<FString>(TEXT("AutoGen Folder (Global)")));
-	AssetLocationModes.Add(MakeShared<FString>(TEXT("Current Folder")));
+	AssetLocationModes.Add(MakeShared<FString>(LOCTEXT("AssetLocationModeAutoGenWorldRelative", "AutoGen Folder (World-Relative)").ToString()));
+	AssetLocationModes.Add(MakeShared<FString>(LOCTEXT("AssetLocationModeAutoGenGlobal", "AutoGen Folder (Global)").ToString()));
+	AssetLocationModes.Add(MakeShared<FString>(LOCTEXT("AssetLocationModeCurrentFolder", "Current Folder").ToString()));
 	AssetLocationMode = SNew(STextComboBox)
 		.OptionsSource(&AssetLocationModes)
 		.OnSelectionChanged_Lambda([&](TSharedPtr<FString> String, ESelectInfo::Type) { UpdateAssetLocationMode(String); });
-	AssetSaveModes.Add(MakeShared<FString>(TEXT("AutoSave New Assets")));
-	AssetSaveModes.Add(MakeShared<FString>(TEXT("Manual Save")));
-	AssetSaveModes.Add(MakeShared<FString>(TEXT("Interactive")));
+	AssetSaveModes.Add(MakeShared<FString>(LOCTEXT("AssetSaveModeAutoSave", "AutoSave New Assets").ToString()));
+	AssetSaveModes.Add(MakeShared<FString>(LOCTEXT("AssetSaveModeManualSave", "Manual Save").ToString()));
+	AssetSaveModes.Add(MakeShared<FString>(LOCTEXT("AssetSaveModeInteractive", "Interactive").ToString()));
 	AssetSaveMode = SNew(STextComboBox)
 		.OptionsSource(&AssetSaveModes)
 		.OnSelectionChanged_Lambda([&](TSharedPtr<FString> String, ESelectInfo::Type) { UpdateAssetSaveMode(String); });
@@ -767,16 +767,16 @@ TSharedPtr<SWidget> FModelingToolsEditorModeToolkit::MakeAssetConfigPanel()
 	// LOD selection dropdown
 	//
 
-	AssetLODModes.Add(MakeShared<FString>(TEXT("Max Available")));
-	AssetLODModes.Add(MakeShared<FString>(TEXT("HiRes")));
-	AssetLODModes.Add(MakeShared<FString>(TEXT("LOD0")));
-	AssetLODModes.Add(MakeShared<FString>(TEXT("LOD1")));
-	AssetLODModes.Add(MakeShared<FString>(TEXT("LOD2")));
-	AssetLODModes.Add(MakeShared<FString>(TEXT("LOD3")));
-	AssetLODModes.Add(MakeShared<FString>(TEXT("LOD4")));
-	AssetLODModes.Add(MakeShared<FString>(TEXT("LOD5")));
-	AssetLODModes.Add(MakeShared<FString>(TEXT("LOD6")));
-	AssetLODModes.Add(MakeShared<FString>(TEXT("LOD7")));
+	AssetLODModes.Add(MakeShared<FString>(LOCTEXT("AssetLODModeMaxAvailable", "Max Available").ToString()));
+	AssetLODModes.Add(MakeShared<FString>(LOCTEXT("AssetLODModeHiRes", "HiRes").ToString()));
+	AssetLODModes.Add(MakeShared<FString>(LOCTEXT("AssetLODModeLOD0", "LOD0").ToString()));
+	AssetLODModes.Add(MakeShared<FString>(LOCTEXT("AssetLODModeLOD1", "LOD1").ToString()));
+	AssetLODModes.Add(MakeShared<FString>(LOCTEXT("AssetLODModeLOD2", "LOD2").ToString()));
+	AssetLODModes.Add(MakeShared<FString>(LOCTEXT("AssetLODModeLOD3", "LOD3").ToString()));
+	AssetLODModes.Add(MakeShared<FString>(LOCTEXT("AssetLODModeLOD4", "LOD4").ToString()));
+	AssetLODModes.Add(MakeShared<FString>(LOCTEXT("AssetLODModeLOD5", "LOD5").ToString()));
+	AssetLODModes.Add(MakeShared<FString>(LOCTEXT("AssetLODModeLOD6", "LOD6").ToString()));
+	AssetLODModes.Add(MakeShared<FString>(LOCTEXT("AssetLODModeLOD7", "LOD7").ToString()));
 	AssetLODMode = SNew(STextComboBox)
 		.OptionsSource(&AssetLODModes)
 		.OnSelectionChanged_Lambda([&](TSharedPtr<FString> String, ESelectInfo::Type)
