enable scaling for geometry selection transform gizmo regardless of frame, and make scaling use the local frame path when the gizmo is in a local frame

diff --git a/Engine/Plugins/Editor/ModelingToolsEditorMode/Source/ModelingToolsEditorMode/Private/ModelingSelectionInteraction.cpp b/Engine/Plugins/Editor/ModelingToolsEditorMode/Source/ModelingToolsEditorMode/Private/ModelingSelectionInteraction.cpp
--- a/Engine/Plugins/Editor/ModelingToolsEditorMode/Source/ModelingToolsEditorMode/Private/ModelingSelectionInteraction.cpp
+++ b/Engine/Plugins/Editor/ModelingToolsEditorMode/Source/ModelingToolsEditorMode/Private/ModelingSelectionInteraction.cpp
@@ -60,6 +60,7 @@ void UModelingSelectionInteraction::Initialize(
 	TransformGizmo = UE::TransformGizmoUtil::CreateCustomRepositionableTransformGizmo( 
 		SelectionManager->GetToolsContext()->GizmoManager, ETransformGizmoSubElements::FullTranslateRotateScale,
 		this, TEXT("ModelingSelectionInteraction") );
+	TransformGizmo->SetIsNonUniformScaleAllowedFunction([]() {return true;});
 	TransformGizmo->SetActiveTarget(TransformProxy, SelectionManager->GetToolsContext()->GizmoManager);
 	TransformGizmo->SetVisibility(false);
 
@@ -412,7 +413,7 @@ void UModelingSelectionInteraction::ApplyPendingTransformInteractions()
 	FQuaterniond RotateDelta = CurFrame.Rotation - InitialGizmoFrame.Rotation;
 	FVector3d CurScaleDelta = CurScale - InitialGizmoScale;
 
-	const bool bLastUpdateUsedWorldFrame = true;	// for later local-frame support
+	bool bTransformInWorldFrame = TransformGizmo->CurrentCoordinateSystem == EToolContextCoordinateSystem::World;
 
 	// if any of the deltas have deviated from the last delta, forward the transformation change on to targets
 	if ( (TranslationDelta - LastTranslationDelta).SizeSquared() > FMathf::ZeroTolerance
@@ -423,7 +424,7 @@ void UModelingSelectionInteraction::ApplyPendingTransformInteractions()
 		LastRotateDelta = (FVector4d)RotateDelta;
 		LastScaleDelta = CurScaleDelta;
 
-		if (bLastUpdateUsedWorldFrame)
+		if (bTransformInWorldFrame)
 		{
 			// For a world frame gizmo, the scaling needs to happen in world aligned gizmo space, but the 
 			// rotation is still encoded in the local gizmo frame change.
