Verify that Polygon Mesh is created using Extrude Polygon Tool

diff --git a/Engine/Plugins/Editor/ModelingToolsEditorMode/Source/ModelingToolsEditorMode/Private/Tests/ExtrudePolygonCoreSmoke/ExtrudePolygon.spec.cpp b/Engine/Plugins/Editor/ModelingToolsEditorMode/Source/ModelingToolsEditorMode/Private/Tests/ExtrudePolygonCoreSmoke/ExtrudePolygon.spec.cpp
new file mode 100644
--- /dev/null
+++ b/Engine/Plugins/Editor/ModelingToolsEditorMode/Source/ModelingToolsEditorMode/Private/Tests/ExtrudePolygonCoreSmoke/ExtrudePolygon.spec.cpp
@@ -0,0 +1,108 @@
+// Copyright Epic Games, Inc. All Rights Reserved.
+
+#include "CoreMinimal.h"
+#include "Engine/World.h"
+#include "Engine/StaticMeshActor.h"
+#include "Components/StaticMeshComponent.h"
+#include "Misc/AutomationTest.h"
+#include "EngineUtils.h"
+#include "ObjectTools.h"
+#include "EditorModeManager.h"
+#include "EditorModes.h"
+#include "ModelingToolsEditorMode.h"
+#include "DrawPolygonTool.h"
+
+BEGIN_DEFINE_SPEC(
+	FExtrudePolygonSpec, "Plugins.Editor.ModelingToolsEditorMode.ExtrudePolygonCoreSmoke",
+	EAutomationTestFlags::EditorContext | EAutomationTestFlags::EngineFilter)
+
+	UEdMode* EditorMode;
+	UInteractiveToolManager* InteractiveToolManager;
+	UDrawPolygonTool* DrawPolygonTool;
+	UWorld* World;
+	AStaticMeshActor* PolygonMesh;
+
+END_DEFINE_SPEC(FExtrudePolygonSpec)
+void FExtrudePolygonSpec::Define()
+{
+	Describe("ExtrudePolygon", [this]()
+	{
+		BeforeEach([this]()
+		{
+			// Entering Extrude Polygon Tool
+			GLevelEditorModeTools().ActivateMode(UModelingToolsEditorMode::EM_ModelingToolsEditorModeId);
+			EditorMode = GLevelEditorModeTools().GetActiveScriptableMode(UModelingToolsEditorMode::EM_ModelingToolsEditorModeId);
+			if (EditorMode)
+			{
+				InteractiveToolManager = EditorMode->GetToolManager();
+				World = EditorMode->GetWorld();
+			}
+			if (InteractiveToolManager)
+			{
+				InteractiveToolManager->SelectActiveToolType(EToolSide::Left, TEXT("BeginDrawPolygonTool"));
+				InteractiveToolManager->ActivateTool(EToolSide::Left);
+				DrawPolygonTool = Cast<UDrawPolygonTool>(InteractiveToolManager->GetActiveTool(EToolSide::Left));
+			}
+		});
+
+		It("Should create Polygon Mesh using Extrude Polygon Tool", [this]()
+		{
+			// We use 3 arbitrary vertices, based on which we'll create the Polygon Mesh
+			const FVector3d VertexA(0.0f, 0.0f, 0.0f);
+			const FVector3d VertexB(0.0f, 100.0f, 0.0f);
+			const FVector3d VertexC(100.0f, 0.0f, 0.0f);
+
+			if (DrawPolygonTool)
+			{
+				DrawPolygonTool->AppendVertex(VertexA);
+				DrawPolygonTool->AppendVertex(VertexB);
+				DrawPolygonTool->AppendVertex(VertexC);
+				DrawPolygonTool->EmitCurrentPolygon();
+			}
+
+			// Test requires that Static Mesh created should be named "Extrude"
+			const FString PolygonName = TEXT("Extrude");
+			bool bPolygonExists = false;
+
+			if (World)
+			{
+				for (AActor* Actor : TActorRange<AActor>(World))
+				{
+					AStaticMeshActor* StaticMeshActor = Cast<AStaticMeshActor>(Actor);
+					if (StaticMeshActor && StaticMeshActor->GetActorLabel() == PolygonName)
+					{
+						PolygonMesh = StaticMeshActor;
+						bPolygonExists = true;
+					}
+				}
+			}
+
+			TestTrue(TEXT("StaticMeshActor named \"Extrude\" does not exist"), bPolygonExists);
+		});
+
+		AfterEach([this]()
+		{
+			// Grabbing Static Mesh of PolygonMesh before we remove the actor
+			UStaticMesh* StaticMesh = nullptr;
+			if (PolygonMesh)
+			{
+				if (const UStaticMeshComponent* StaticMeshComponent = PolygonMesh->GetStaticMeshComponent())
+				{
+					StaticMesh = StaticMeshComponent->GetStaticMesh();
+				}
+				// Removing PolygonMesh Actor
+				World->EditorDestroyActor(PolygonMesh, false);
+			}
+			// Cleaning up temporary Static Mesh that doesn't get removed by removing the Actor
+			if (StaticMesh)
+			{
+				TArray<UObject*> AssetsToDelete;
+				AssetsToDelete.Add(StaticMesh);
+				ObjectTools::DeleteObjects(AssetsToDelete, false, ObjectTools::EAllowCancelDuringDelete::CancelNotAllowed);
+			}
+			// Resetting the interface out of Modeling Tools
+			InteractiveToolManager->DeactivateTool(EToolSide::Left, EToolShutdownType::Completed);
+			GLevelEditorModeTools().ActivateMode(FBuiltinEditorModes::EM_Default);
+		});
+	});
+}
