Verify proper removal and addition of Polygon Mesh via undo-redo

diff --git a/Engine/Plugins/Editor/ModelingToolsEditorMode/Source/ModelingToolsEditorMode/Private/Tests/ExtrudePolygonCoreSmoke/ExtrudePolygon.spec.cpp b/Engine/Plugins/Editor/ModelingToolsEditorMode/Source/ModelingToolsEditorMode/Private/Tests/ExtrudePolygonCoreSmoke/ExtrudePolygon.spec.cpp
--- a/Engine/Plugins/Editor/ModelingToolsEditorMode/Source/ModelingToolsEditorMode/Private/Tests/ExtrudePolygonCoreSmoke/ExtrudePolygon.spec.cpp
+++ b/Engine/Plugins/Editor/ModelingToolsEditorMode/Source/ModelingToolsEditorMode/Private/Tests/ExtrudePolygonCoreSmoke/ExtrudePolygon.spec.cpp
@@ -13,7 +13,7 @@
 #include "DrawPolygonTool.h"
 
 BEGIN_DEFINE_SPEC(
-	FExtrudePolygonSpec, "Plugins.Editor.ModelingToolsEditorMode.ExtrudePolygonCoreSmoke",
+	FExtrudePolygonSpec, "Editor.Plugins.Tools.Modeling.EditorMode",
 	EAutomationTestFlags::EditorContext | EAutomationTestFlags::EngineFilter)
 
 	UEdMode* EditorMode;
@@ -21,6 +21,9 @@ BEGIN_DEFINE_SPEC(
 	UDrawPolygonTool* DrawPolygonTool;
 	UWorld* World;
 	AStaticMeshActor* PolygonMesh;
+	// Test helper functions
+	bool MeshExists(const FString&);
+	void CreateMesh() const;
 
 END_DEFINE_SPEC(FExtrudePolygonSpec)
 void FExtrudePolygonSpec::Define()
@@ -47,37 +50,33 @@ void FExtrudePolygonSpec::Define()
 
 		It("Should create Polygon Mesh using Extrude Polygon Tool", [this]()
 		{
-			// We use 3 arbitrary vertices, based on which we'll create the Polygon Mesh
-			const FVector3d VertexA(0.0f, 0.0f, 0.0f);
-			const FVector3d VertexB(0.0f, 100.0f, 0.0f);
-			const FVector3d VertexC(100.0f, 0.0f, 0.0f);
+			CreateMesh();
 
-			if (DrawPolygonTool)
-			{
-				DrawPolygonTool->AppendVertex(VertexA);
-				DrawPolygonTool->AppendVertex(VertexB);
-				DrawPolygonTool->AppendVertex(VertexC);
-				DrawPolygonTool->EmitCurrentPolygon();
-			}
+			// Test requires that Static Mesh created should be named "Extrude"
+			const FString PolygonName = TEXT("Extrude");
+
+			const bool bPolygonExists = MeshExists(PolygonName);
+
+			TestTrue(TEXT("StaticMeshActor named \"Extrude\" exists"), bPolygonExists);
+		});
+
+		It("Should validate proper removal and addition of Polygon Mesh via undo-redo", [this]()
+		{
+			CreateMesh();
 
 			// Test requires that Static Mesh created should be named "Extrude"
 			const FString PolygonName = TEXT("Extrude");
-			bool bPolygonExists = false;
 
-			if (World)
+			GEditor->UndoTransaction();
+			const bool bMeshUndoState = MeshExists(PolygonName);
+			
+			if (TestTrue(TEXT("Polygon Mesh absent after Undo"), !bMeshUndoState))
 			{
-				for (AActor* Actor : TActorRange<AActor>(World))
-				{
-					AStaticMeshActor* StaticMeshActor = Cast<AStaticMeshActor>(Actor);
-					if (StaticMeshActor && StaticMeshActor->GetActorLabel() == PolygonName)
-					{
-						PolygonMesh = StaticMeshActor;
-						bPolygonExists = true;
-					}
-				}
+				// Proceed with Redo only if the Undo operation was successful
+				GEditor->RedoTransaction();
+				const bool bMeshRedoState = MeshExists(PolygonName);
+				TestTrue(TEXT("Polygon Mesh restored after Redo"), bMeshRedoState);
 			}
-
-			TestTrue(TEXT("StaticMeshActor named \"Extrude\" does not exist"), bPolygonExists);
 		});
 
 		AfterEach([this]()
@@ -106,3 +105,39 @@ void FExtrudePolygonSpec::Define()
 		});
 	});
 }
+
+void FExtrudePolygonSpec::CreateMesh() const
+{
+	// We use 3 arbitrary vertices, based on which we'll create the Polygon Mesh
+	const FVector3d VertexA(0.0f, 0.0f, 0.0f);
+	const FVector3d VertexB(0.0f, 100.0f, 0.0f);
+	const FVector3d VertexC(100.0f, 0.0f, 0.0f);
+
+	if (DrawPolygonTool)
+	{
+		DrawPolygonTool->AppendVertex(VertexA);
+		DrawPolygonTool->AppendVertex(VertexB);
+		DrawPolygonTool->AppendVertex(VertexC);
+		DrawPolygonTool->EmitCurrentPolygon();
+	}
+}
+
+bool FExtrudePolygonSpec::MeshExists(const FString& PolygonName)
+{
+	bool bMeshExists = false;
+
+	if (World)
+	{
+		for (AActor* Actor : TActorRange<AActor>(World))
+		{
+			AStaticMeshActor* StaticMeshActor = Cast<AStaticMeshActor>(Actor);
+			if (StaticMeshActor && StaticMeshActor->GetActorLabel() == PolygonName)
+			{
+				bMeshExists = true;
+				PolygonMesh = StaticMeshActor;
+			}
+		}
+	}
+
+	return bMeshExists;
+}
\ No newline at end of file
