Verify that completing after creating a vertex dismisses Extrude Polygon Settings

diff --git a/Engine/Plugins/Editor/ModelingToolsEditorMode/Source/ModelingToolsEditorMode/Private/Tests/ExtrudePolygonCoreSmoke/ExtrudePolygon.spec.cpp b/Engine/Plugins/Editor/ModelingToolsEditorMode/Source/ModelingToolsEditorMode/Private/Tests/ExtrudePolygonCoreSmoke/ExtrudePolygon.spec.cpp
--- a/Engine/Plugins/Editor/ModelingToolsEditorMode/Source/ModelingToolsEditorMode/Private/Tests/ExtrudePolygonCoreSmoke/ExtrudePolygon.spec.cpp
+++ b/Engine/Plugins/Editor/ModelingToolsEditorMode/Source/ModelingToolsEditorMode/Private/Tests/ExtrudePolygonCoreSmoke/ExtrudePolygon.spec.cpp
@@ -22,8 +22,9 @@ BEGIN_DEFINE_SPEC(
 	UWorld* World;
 	AStaticMeshActor* PolygonMesh;
 	// Test helper functions
-	bool MeshExists(const FString&);
 	void CreateMesh() const;
+	bool MeshExists(const FString&);
+	bool ModelingSettingsExist(const FString&) const;
 
 END_DEFINE_SPEC(FExtrudePolygonSpec)
 void FExtrudePolygonSpec::Define()
@@ -79,6 +80,25 @@ void FExtrudePolygonSpec::Define()
 			}
 		});
 
+		It("Should dismiss Extrude Polygon Settings after creating a vertex and completing", [this]()
+		{
+			const FString ExtrudePolygonToolName = TEXT("BeginDrawPolygonTool");
+			// Checking that during setup we entered Extrude Polygon Tool
+			if (!TestTrue(TEXT("Entered Extrude Polygon Tool successfully"), ModelingSettingsExist(ExtrudePolygonToolName)))
+			{
+				return;
+			}
+
+			if (InteractiveToolManager && DrawPolygonTool)
+			{
+				DrawPolygonTool->AppendVertex(FVector3d::ZeroVector);
+				// Clicking "Complete" will invoke InteractiveToolManager->DeactivateTool
+				InteractiveToolManager->DeactivateTool(EToolSide::Left, EToolShutdownType::Completed);
+			}
+
+			TestTrue(TEXT("Extrude Polygon Settings dismissed after completing"), !ModelingSettingsExist(ExtrudePolygonToolName));
+		});
+
 		AfterEach([this]()
 		{
 			// Grabbing Static Mesh of PolygonMesh before we remove the actor
@@ -140,4 +160,14 @@ bool FExtrudePolygonSpec::MeshExists(const FString& PolygonName)
 	}
 
 	return bMeshExists;
-}
\ No newline at end of file
+}
+
+bool FExtrudePolygonSpec::ModelingSettingsExist(const FString& ToolName) const
+{
+	const FString ActiveToolName = InteractiveToolManager->GetActiveToolName(EToolSide::Left);
+	if (ActiveToolName == ToolName)
+	{
+		return true;
+	}
+	return false;
+}
