Verify that upon entering Extrude Polygon Tool, the modeling tab populates with Extrude Polygon tool settings

diff --git a/Engine/Plugins/Editor/ModelingToolsEditorMode/Source/ModelingToolsEditorMode/Private/Tests/ExtrudePolygonCoreSmoke/ExtrudePolygon.spec.cpp b/Engine/Plugins/Editor/ModelingToolsEditorMode/Source/ModelingToolsEditorMode/Private/Tests/ExtrudePolygonCoreSmoke/ExtrudePolygon.spec.cpp
--- a/Engine/Plugins/Editor/ModelingToolsEditorMode/Source/ModelingToolsEditorMode/Private/Tests/ExtrudePolygonCoreSmoke/ExtrudePolygon.spec.cpp
+++ b/Engine/Plugins/Editor/ModelingToolsEditorMode/Source/ModelingToolsEditorMode/Private/Tests/ExtrudePolygonCoreSmoke/ExtrudePolygon.spec.cpp
@@ -99,6 +99,13 @@ void FExtrudePolygonSpec::Define()
 			TestTrue(TEXT("Extrude Polygon Settings dismissed after completing"), !ModelingSettingsExist(ExtrudePolygonToolName));
 		});
 
+		It("Should populate Modeling tab with Extrude Polygon settings after tool entry", [this]()
+		{
+			const FString ExtrudePolygonToolName = TEXT("BeginDrawPolygonTool");
+			TestTrue(TEXT("The Modeling tab is populated with Extrude Polygon settings"),
+					 ModelingSettingsExist(ExtrudePolygonToolName));
+		});
+
 		AfterEach([this]()
 		{
 			// Grabbing Static Mesh of PolygonMesh before we remove the actor
