Exposes modeling tools editor ID.

diff --git a/Engine/Plugins/Editor/ModelingToolsEditorMode/Source/ModelingToolsEditorMode/Public/ModelingToolsEditorMode.h b/Engine/Plugins/Editor/ModelingToolsEditorMode/Source/ModelingToolsEditorMode/Public/ModelingToolsEditorMode.h
--- a/Engine/Plugins/Editor/ModelingToolsEditorMode/Source/ModelingToolsEditorMode/Public/ModelingToolsEditorMode.h
+++ b/Engine/Plugins/Editor/ModelingToolsEditorMode/Source/ModelingToolsEditorMode/Public/ModelingToolsEditorMode.h
@@ -28,7 +28,7 @@ class UModelingToolsEditorMode : public UBaseLegacyWidgetEdMode, public ILegacyE
 {
 	GENERATED_BODY()
 public:
-	const static FEditorModeID EM_ModelingToolsEditorModeId;
+	MODELINGTOOLSEDITORMODE_API const static FEditorModeID EM_ModelingToolsEditorModeId;
 
 	UModelingToolsEditorMode();
 	UModelingToolsEditorMode(FVTableHelper& Helper);
