ModelingTools: Allow tools to customize the accept/cancel overlay for sub-activitites, and use this in PolyEd.

diff --git a/Engine/Plugins/Editor/ModelingToolsEditorMode/Source/ModelingToolsEditorMode/Private/ModelingToolsEditorMode.cpp b/Engine/Plugins/Editor/ModelingToolsEditorMode/Source/ModelingToolsEditorMode/Private/ModelingToolsEditorMode.cpp
--- a/Engine/Plugins/Editor/ModelingToolsEditorMode/Source/ModelingToolsEditorMode/Private/ModelingToolsEditorMode.cpp
+++ b/Engine/Plugins/Editor/ModelingToolsEditorMode/Source/ModelingToolsEditorMode/Private/ModelingToolsEditorMode.cpp
@@ -10,6 +10,7 @@
 #include "ModelingToolsEditorModeToolkit.h"
 #include "ILevelEditor.h"
 #include "ModelingToolsEditorModeSettings.h"
+#include "ModelingToolsHostCustomizationAPI.h"
 #include "ToolTargetManager.h"
 #include "ContextObjectStore.h"
 #include "InputRouter.h"
@@ -428,6 +429,12 @@ void UModelingToolsEditorMode::Enter()
 	UE::Geometry::RegisterSceneSnappingManager(GetInteractiveToolsContext());
 	SceneSnappingManager = UE::Geometry::FindModelingSceneSnappingManager(GetToolManager());
 
+	// register tool shutdown button customizer
+	if (ensure(Toolkit.IsValid()))
+	{
+		UModelingToolsHostCustomizationAPI::Register(GetInteractiveToolsContext(), 
+			StaticCastSharedRef<FModelingToolsEditorModeToolkit>(Toolkit.ToSharedRef()));
+	}
 
 	// set up SelectionManager and register known factory types
 	SelectionManager = NewObject<UGeometrySelectionManager>(GetToolManager());
@@ -1241,6 +1248,8 @@ void UModelingToolsEditorMode::Exit()
 	StylusStateTracker = nullptr;
 #endif
 
+	UModelingToolsHostCustomizationAPI::Deregister(GetInteractiveToolsContext());
+
 	// TODO: cannot deregister currently because if another mode is also registering, its Enter()
 	// will be called before our Exit()
 	//UE::TransformGizmoUtil::DeregisterTransformGizmoContextObject(ToolsContext.Get());
diff --git a/Engine/Plugins/Editor/ModelingToolsEditorMode/Source/ModelingToolsEditorMode/Private/ModelingToolsEditorModeStyle.cpp b/Engine/Plugins/Editor/ModelingToolsEditorMode/Source/ModelingToolsEditorMode/Private/ModelingToolsEditorModeStyle.cpp
--- a/Engine/Plugins/Editor/ModelingToolsEditorMode/Source/ModelingToolsEditorMode/Private/ModelingToolsEditorModeStyle.cpp
+++ b/Engine/Plugins/Editor/ModelingToolsEditorMode/Source/ModelingToolsEditorMode/Private/ModelingToolsEditorModeStyle.cpp
@@ -98,6 +98,9 @@ void FModelingToolsEditorModeStyle::Initialize()
 		}
 
 		StyleSet->Set("ModelingMode.DefaultSettings", new FSlateImageBrush(StyleSet->RootToCoreContentDir(TEXT("../Editor/Slate/Icons/GeneralTools/Settings_40x.png")), Icon20x20));
+		StyleSet->Set("ModelingMode.SubToolArrow", new FSlateImageBrush(StyleSet->RootToCoreContentDir(
+			// This arrow is 25x13, but 20x13 looks better
+			TEXT("../Editor/Slate/Persona/BlendSpace/arrow_right_12x.png")), FVector2D(20,13)));
 
 		// NOTE:  Old-style, need to be replaced: 
 		StyleSet->Set("ModelingToolsManagerCommands.CancelActiveTool", new IMAGE_PLUGIN_BRUSH("Icons/icon_ActiveTool_Cancel_40x", Icon20x20));
@@ -173,10 +176,6 @@ void FModelingToolsEditorModeStyle::Initialize()
 		StyleSet->Set("ModelingToolsManagerCommands.BeginPolyEditTool.Small", 			new IMAGE_PLUGIN_BRUSH("Icons/PolyEdit_40x", 		Icon20x20));
 		StyleSet->Set("ModelingToolsManagerCommands.BeginSubdividePolyTool",			new IMAGE_BRUSH_SVG("Icons/ModelingSubD",		Icon20x20));
 		StyleSet->Set("ModelingToolsManagerCommands.BeginSubdividePolyTool.Small",		new IMAGE_BRUSH_SVG("Icons/ModelingSubD",		Icon20x20));
-		StyleSet->Set("ModelingToolsManagerCommands.BeginGroupEdgeInsertionTool", 		new IMAGE_PLUGIN_BRUSH("Icons/ModelingGroupEdgeInsert_x40", 		Icon20x20));
-		StyleSet->Set("ModelingToolsManagerCommands.BeginGroupEdgeInsertionTool.Small", new IMAGE_PLUGIN_BRUSH("Icons/ModelingGroupEdgeInsert_x40", 		Icon20x20));
-		StyleSet->Set("ModelingToolsManagerCommands.BeginEdgeLoopInsertionTool", 		new IMAGE_PLUGIN_BRUSH("Icons/ModelingEdgeLoopInsert_x40", 		Icon20x20));
-		StyleSet->Set("ModelingToolsManagerCommands.BeginEdgeLoopInsertionTool.Small",  new IMAGE_PLUGIN_BRUSH("Icons/ModelingEdgeLoopInsert_x40", 		Icon20x20));
 		StyleSet->Set("ModelingToolsManagerCommands.BeginTriEditTool", 				new IMAGE_PLUGIN_BRUSH("Icons/TriEdit_40x", 		Icon20x20));
 		StyleSet->Set("ModelingToolsManagerCommands.BeginTriEditTool.Small", 			new IMAGE_PLUGIN_BRUSH("Icons/TriEdit_40x", 		Icon20x20));
 		//StyleSet->Set("ModelingToolsManagerCommands.BeginPolyDeformTool", 			new IMAGE_PLUGIN_BRUSH("Icons/PolyEdit_40x", 		Icon20x20));
@@ -460,6 +459,10 @@ void FModelingToolsEditorModeStyle::Initialize()
 		StyleSet->Set("PolyEd.SelectFaces", new IMAGE_BRUSH_SVG("Icons/SelectionTriangles3", Icon20x20));
 		StyleSet->Set("PolyEd.SelectEdgeLoops", new IMAGE_BRUSH_SVG("Icons/ModelingEdgeLoopSelection", Icon20x20));
 		StyleSet->Set("PolyEd.SelectEdgeRings", new IMAGE_BRUSH_SVG("Icons/ModelingEdgeRingSelection", Icon20x20));
+
+		// Icons for PolyEd and TriEd activities
+		StyleSet->Set("PolyEd.InsertGroupEdge", new IMAGE_PLUGIN_BRUSH("Icons/ModelingGroupEdgeInsert_x40", Icon20x20));
+		StyleSet->Set("PolyEd.InsertEdgeLoop", new IMAGE_PLUGIN_BRUSH("Icons/ModelingEdgeLoopInsert_x40", Icon20x20));
 	}
 
 	// Style for the toolbar in the PolyEd customization
diff --git a/Engine/Plugins/Editor/ModelingToolsEditorMode/Source/ModelingToolsEditorMode/Private/ModelingToolsEditorModeToolkit.cpp b/Engine/Plugins/Editor/ModelingToolsEditorMode/Source/ModelingToolsEditorMode/Private/ModelingToolsEditorModeToolkit.cpp
--- a/Engine/Plugins/Editor/ModelingToolsEditorMode/Source/ModelingToolsEditorMode/Private/ModelingToolsEditorModeToolkit.cpp
+++ b/Engine/Plugins/Editor/ModelingToolsEditorMode/Source/ModelingToolsEditorMode/Private/ModelingToolsEditorModeToolkit.cpp
@@ -653,6 +653,119 @@ void FModelingToolsEditorModeToolkit::MakeToolShutdownOverlayWidget()
 		OverlayBrush = FModelingToolsEditorModeStyle::Get()->GetBrush("ModelingMode.OpaqueOverlayBrush");
 	}
 
+	// Helpers to determine button/label visibility based on overrides
+	auto GetSubActionIcon = [this]() -> const FSlateBrush*
+	{
+		if (AcceptCancelButtonParams.IsSet())
+		{
+			if (AcceptCancelButtonParams->IconName.IsSet())
+			{
+				return FModelingToolsEditorModeStyle::Get()->GetOptionalBrush(AcceptCancelButtonParams->IconName.GetValue(), nullptr, nullptr);
+			}
+		}
+		else if (CompleteButtonParams.IsSet() && CompleteButtonParams->IconName.IsSet())
+		{
+			return FModelingToolsEditorModeStyle::Get()->GetOptionalBrush(CompleteButtonParams->IconName.GetValue(), nullptr, nullptr);
+		}
+		return nullptr;
+	};
+	auto GetSubActionIconVisibility = [this]()
+	{
+		if (AcceptCancelButtonParams.IsSet())
+		{
+			if (AcceptCancelButtonParams->IconName.IsSet() 
+				&& FModelingToolsEditorModeStyle::Get()->GetOptionalBrush(AcceptCancelButtonParams->IconName.GetValue(), nullptr, nullptr))
+			{
+				return EVisibility::Visible;
+			}
+		}
+		else if (CompleteButtonParams.IsSet() && CompleteButtonParams->IconName.IsSet()
+			&& FModelingToolsEditorModeStyle::Get()->GetOptionalBrush(CompleteButtonParams->IconName.GetValue(), nullptr, nullptr))
+		{
+			return EVisibility::Visible;
+		}
+		return EVisibility::Collapsed;
+	};
+	auto GetSubActionLabel = [this]()
+	{
+		return AcceptCancelButtonParams.IsSet() ? AcceptCancelButtonParams->Label
+			: CompleteButtonParams.IsSet() ? CompleteButtonParams->Label
+			: FText::GetEmpty();
+	};
+	auto GetSubActionLabelVisibility = [this]()
+	{
+		return (AcceptCancelButtonParams.IsSet() || CompleteButtonParams.IsSet()) ?
+			EVisibility::Visible : EVisibility::Collapsed;
+	};
+	auto GetAcceptButtonText = [this]()
+	{
+		return AcceptCancelButtonParams.IsSet() && AcceptCancelButtonParams->OverrideAcceptButtonText.IsSet() ?
+			AcceptCancelButtonParams->OverrideAcceptButtonText.GetValue() : LOCTEXT("OverlayAccept", "Accept");
+	};
+	auto GetAcceptButtonTooltip = [this]()
+	{
+		return AcceptCancelButtonParams.IsSet() && AcceptCancelButtonParams->OverrideAcceptButtonTooltip.IsSet() ?
+			AcceptCancelButtonParams->OverrideAcceptButtonTooltip.GetValue()
+			: LOCTEXT("OverlayAcceptTooltip", "Accept/Commit the results of the active Tool [Enter]");
+	};
+	auto GetAcceptButtonEnabled = [this]()
+	{
+		return AcceptCancelButtonParams.IsSet() ? AcceptCancelButtonParams->CanAccept()
+			: GetScriptableEditorMode()->GetInteractiveToolsContext(EToolsContextScope::EdMode)->CanAcceptActiveTool();
+	};
+	auto GetAcceptCancelButtonVisibility = [this]()
+	{
+		if (AcceptCancelButtonParams.IsSet()
+			|| (!CompleteButtonParams.IsSet()
+				&& GetScriptableEditorMode()->GetInteractiveToolsContext(EToolsContextScope::EdMode)->ActiveToolHasAccept()))
+		{
+			return EVisibility::Visible;
+		}
+		return EVisibility::Collapsed;
+	};
+	auto GetCancelButtonText = [this]()
+	{
+		return AcceptCancelButtonParams.IsSet() && AcceptCancelButtonParams->OverrideCancelButtonText.IsSet() ?
+			AcceptCancelButtonParams->OverrideCancelButtonText.GetValue() : LOCTEXT("OverlayCancel", "Cancel");
+	};
+	auto GetCancelButtonTooltip = [this]()
+	{
+		return AcceptCancelButtonParams.IsSet() && AcceptCancelButtonParams->OverrideCancelButtonTooltip.IsSet() ?
+			AcceptCancelButtonParams->OverrideCancelButtonTooltip.GetValue()
+			: LOCTEXT("OverlayCancelTooltip", "Cancel the active Tool [Esc]");
+	};
+	auto GetCancelButtonEnabled = [this]() 
+	{
+		return AcceptCancelButtonParams.IsSet()
+			|| GetScriptableEditorMode()->GetInteractiveToolsContext(EToolsContextScope::EdMode)->CanCancelActiveTool();
+	};
+	auto GetCompleteButtonText = [this]()
+	{
+		return CompleteButtonParams.IsSet() && CompleteButtonParams->OverrideCompleteButtonText.IsSet() ?
+			CompleteButtonParams->OverrideCompleteButtonText.GetValue() : LOCTEXT("OverlayComplete", "Complete");
+	};
+	auto GetCompleteButtonTooltip = [this]()
+	{
+		return CompleteButtonParams.IsSet() && CompleteButtonParams->OverrideCompleteButtonTooltip.IsSet() ?
+			CompleteButtonParams->OverrideCompleteButtonTooltip.GetValue()
+			: LOCTEXT("OverlayCompleteTooltip", "Exit the active Tool [Enter]");
+	};
+	auto GetCompleteButtonEnabled = [this]()
+	{
+		return CompleteButtonParams.IsSet()
+			|| GetScriptableEditorMode()->GetInteractiveToolsContext(EToolsContextScope::EdMode)->CanCompleteActiveTool();
+	};
+	auto GetCompleteButtonVisibility = [this]()
+	{
+		if (CompleteButtonParams.IsSet()
+			|| (!AcceptCancelButtonParams.IsSet()
+				&& GetScriptableEditorMode()->GetInteractiveToolsContext(EToolsContextScope::EdMode)->CanCompleteActiveTool()))
+		{
+			return EVisibility::Visible;
+		}
+		return EVisibility::Collapsed;
+	};
+
 	SAssignNew(ToolShutdownViewportOverlayWidget, SHorizontalBox)
 
 	+SHorizontalBox::Slot()
@@ -666,6 +779,7 @@ void FModelingToolsEditorModeToolkit::MakeToolShutdownOverlayWidget()
 		[
 			SNew(SHorizontalBox)
 
+			// Tool icon and name
 			+SHorizontalBox::Slot()
 			.AutoWidth()
 			.VAlign(VAlign_Center)
@@ -674,7 +788,6 @@ void FModelingToolsEditorModeToolkit::MakeToolShutdownOverlayWidget()
 				SNew(SImage)
 				.Image_Lambda([this] () { return ActiveToolIcon; })
 			]
-
 			+SHorizontalBox::Slot()
 			.AutoWidth()
 			.VAlign(VAlign_Center)
@@ -684,41 +797,76 @@ void FModelingToolsEditorModeToolkit::MakeToolShutdownOverlayWidget()
 				.Text(this, &FModelingToolsEditorModeToolkit::GetActiveToolDisplayName)
 			]
 
+			// Optional: "-> [icon] SubtoolAction"
+			// arrow
+			+SHorizontalBox::Slot()
+			.AutoWidth()
+			.VAlign(EVerticalAlignment::VAlign_Center)
+			.Padding(FMargin(0., 0.f, 8.f, 0.f))
+			[
+				SNew(SImage)
+				.Image(FModelingToolsEditorModeStyle::Get()->GetBrush("ModelingMode.SubToolArrow"))
+				.ColorAndOpacity(FSlateColor::UseForeground())
+				.Visibility_Lambda(GetSubActionLabelVisibility)
+			]
+			// subaction icon
+			+SHorizontalBox::Slot()
+			.AutoWidth()
+			.VAlign(VAlign_Center)
+			.Padding(FMargin(0.f, 0.f, 8.f, 0.f))
+			[
+				SNew(SImage)
+				.Image_Lambda(GetSubActionIcon)
+				.Visibility_Lambda(GetSubActionIconVisibility)
+			]
+			// subaction label
+			+SHorizontalBox::Slot()
+			.AutoWidth()
+			.VAlign(VAlign_Center)
+			.Padding(FMargin(0.f, 0.f, 8.f, 0.f))
+			[
+				SNew(STextBlock)
+				.Text_Lambda(GetSubActionLabel)
+				.Visibility_Lambda(GetSubActionLabelVisibility)
+			]
+
+			// Buttons:
+			// Accept
 			+SHorizontalBox::Slot()
 			.AutoWidth()
 			.Padding(FMargin(0.0, 0.f, 2.f, 0.f))
 			[
 				SNew(SPrimaryButton)
-				.Text(LOCTEXT("OverlayAccept", "Accept"))
-				.ToolTipText(LOCTEXT("OverlayAcceptTooltip", "Accept/Commit the results of the active Tool [Enter]"))
-				.OnClicked_Lambda([this]() { GetScriptableEditorMode()->GetInteractiveToolsContext(EToolsContextScope::EdMode)->EndTool(EToolShutdownType::Accept); return FReply::Handled(); })
-				.IsEnabled_Lambda([this]() { return GetScriptableEditorMode()->GetInteractiveToolsContext(EToolsContextScope::EdMode)->CanAcceptActiveTool(); })
-				.Visibility_Lambda([this]() { return GetScriptableEditorMode()->GetInteractiveToolsContext(EToolsContextScope::EdMode)->ActiveToolHasAccept() ? EVisibility::Visible : EVisibility::Collapsed; })
+				.Text_Lambda(GetAcceptButtonText)
+				.ToolTipText_Lambda(GetAcceptButtonTooltip)
+				.OnClicked_Raw(this, &FModelingToolsEditorModeToolkit::HandleAcceptCancelClick, true)
+				.IsEnabled_Lambda(GetAcceptButtonEnabled)
+				.Visibility_Lambda(GetAcceptCancelButtonVisibility)
 			]
-
+			// Cancel
 			+SHorizontalBox::Slot()
 			.AutoWidth()
 			.Padding(FMargin(2.0, 0.f, 0.f, 0.f))
 			[
 				SNew(SButton)
-				.Text(LOCTEXT("OverlayCancel", "Cancel"))
-				.ToolTipText(LOCTEXT("OverlayCancelTooltip", "Cancel the active Tool [Esc]"))
+				.Text_Lambda(GetCancelButtonText)
+				.ToolTipText_Lambda(GetCancelButtonTooltip)
 				.HAlign(HAlign_Center)
-				.OnClicked_Lambda([this]() { GetScriptableEditorMode()->GetInteractiveToolsContext(EToolsContextScope::EdMode)->EndTool(EToolShutdownType::Cancel); return FReply::Handled(); })
-				.IsEnabled_Lambda([this]() { return GetScriptableEditorMode()->GetInteractiveToolsContext(EToolsContextScope::EdMode)->CanCancelActiveTool(); })
-				.Visibility_Lambda([this]() { return GetScriptableEditorMode()->GetInteractiveToolsContext(EToolsContextScope::EdMode)->ActiveToolHasAccept() ? EVisibility::Visible : EVisibility::Collapsed; })
+				.OnClicked_Raw(this, &FModelingToolsEditorModeToolkit::HandleAcceptCancelClick, false)
+				.IsEnabled_Lambda(GetCancelButtonEnabled)
+				.Visibility_Lambda(GetAcceptCancelButtonVisibility)
 			]
-
+			// Complete
 			+SHorizontalBox::Slot()
 			.AutoWidth()
 			.Padding(FMargin(2.0, 0.f, 0.f, 0.f))
 			[
 				SNew(SPrimaryButton)
-				.Text(LOCTEXT("OverlayComplete", "Complete"))
-				.ToolTipText(LOCTEXT("OverlayCompleteTooltip", "Exit the active Tool [Enter]"))
-				.OnClicked_Lambda([this]() { GetScriptableEditorMode()->GetInteractiveToolsContext(EToolsContextScope::EdMode)->EndTool(EToolShutdownType::Completed); return FReply::Handled(); })
-				.IsEnabled_Lambda([this]() { return GetScriptableEditorMode()->GetInteractiveToolsContext(EToolsContextScope::EdMode)->CanCompleteActiveTool(); })
-				.Visibility_Lambda([this]() { return GetScriptableEditorMode()->GetInteractiveToolsContext(EToolsContextScope::EdMode)->CanCompleteActiveTool() ? EVisibility::Visible : EVisibility::Collapsed; })
+				.Text_Lambda(GetCompleteButtonText)
+				.ToolTipText_Lambda(GetCompleteButtonTooltip)
+				.OnClicked_Raw(this, &FModelingToolsEditorModeToolkit::HandleCompleteClick)
+				.IsEnabled_Lambda(GetCompleteButtonEnabled)
+				.Visibility_Lambda(GetCompleteButtonVisibility)
 			]
 		]	
 	];
@@ -2416,6 +2564,98 @@ void FModelingToolsEditorModeToolkit::NotifySelectionSystemEnabledStateModified(
 	//}
 }
 
+bool FModelingToolsEditorModeToolkit::RequestAcceptCancelButtonOverride(IToolHostCustomizationAPI::FAcceptCancelButtonOverrideParams& Params)
+{
+	if (!Params.OnAcceptCancelTriggered || !Params.CanAccept || Params.Label.IsEmpty())
+	{
+		UE_LOG(LogGeometry, Warning, TEXT("FModelingToolsEditorModeToolkit::RequestAcceptCancelButtonOverride received request with insufficient parameters."));
+		return false;
+	}
+
+	AcceptCancelButtonParams = Params;
+	CompleteButtonParams.Reset();
+	bCurrentOverrideButtonsWereClicked = false;
+
+	if (ToolShutdownViewportOverlayWidget)
+	{
+		ToolShutdownViewportOverlayWidget->Invalidate(EInvalidateWidgetReason::Layout);
+	}
+	return true;
+}
+bool FModelingToolsEditorModeToolkit::RequestCompleteButtonOverride(IToolHostCustomizationAPI::FCompleteButtonOverrideParams& Params)
+{
+	if (!Params.OnCompleteTriggered || Params.Label.IsEmpty())
+	{
+		UE_LOG(LogGeometry, Warning, TEXT("FModelingToolsEditorModeToolkit::RequestCompleteButtonOverride received request with insufficient parameters."));
+		return false;
+	}
+
+	CompleteButtonParams = Params;
+	AcceptCancelButtonParams.Reset();
+	bCurrentOverrideButtonsWereClicked = false;
 
+	if (ToolShutdownViewportOverlayWidget)
+	{
+		ToolShutdownViewportOverlayWidget->Invalidate(EInvalidateWidgetReason::Layout);
+	}
+	return true;
+}
+void FModelingToolsEditorModeToolkit::ClearButtonOverrides()
+{
+	AcceptCancelButtonParams.Reset();
+	CompleteButtonParams.Reset();
+	if (ToolShutdownViewportOverlayWidget)
+	{
+		ToolShutdownViewportOverlayWidget->Invalidate(EInvalidateWidgetReason::Layout);
+	}
+}
+
+
+FReply FModelingToolsEditorModeToolkit::HandleAcceptCancelClick(bool bAccept)
+{
+	if (AcceptCancelButtonParams.IsSet())
+	{
+		bCurrentOverrideButtonsWereClicked = true;
+		if (ensure(AcceptCancelButtonParams->OnAcceptCancelTriggered))
+		{
+			AcceptCancelButtonParams->OnAcceptCancelTriggered(bAccept);
+		}
+		
+		// This will be reset back to false if the callback above triggers another override request.
+		if (bCurrentOverrideButtonsWereClicked)
+		{
+			ClearButtonOverrides();
+		}
+	}
+	else
+	{
+		GetScriptableEditorMode()->GetInteractiveToolsContext(EToolsContextScope::EdMode)->EndTool(
+			bAccept ? EToolShutdownType::Accept : EToolShutdownType::Cancel);
+	}
+	return FReply::Handled();
+}
+FReply FModelingToolsEditorModeToolkit::HandleCompleteClick()
+{
+	if (CompleteButtonParams.IsSet())
+	{
+		bCurrentOverrideButtonsWereClicked = true;
+		if (ensure(CompleteButtonParams->OnCompleteTriggered))
+		{
+			CompleteButtonParams->OnCompleteTriggered();
+		}
+
+		// This will be reset back to false if the callback above triggers another override request.
+		if (bCurrentOverrideButtonsWereClicked)
+		{
+			ClearButtonOverrides();
+		}
+	}
+	else
+	{
+		GetScriptableEditorMode()->GetInteractiveToolsContext(EToolsContextScope::EdMode)->EndTool(EToolShutdownType::Completed);
+	}
+
+	return FReply::Handled();
+}
 
 #undef LOCTEXT_NAMESPACE
diff --git a/Engine/Plugins/Editor/ModelingToolsEditorMode/Source/ModelingToolsEditorMode/Private/ModelingToolsHostCustomizationAPI.cpp b/Engine/Plugins/Editor/ModelingToolsEditorMode/Source/ModelingToolsEditorMode/Private/ModelingToolsHostCustomizationAPI.cpp
new file mode 100644
--- /dev/null
+++ b/Engine/Plugins/Editor/ModelingToolsEditorMode/Source/ModelingToolsEditorMode/Private/ModelingToolsHostCustomizationAPI.cpp
@@ -0,0 +1,75 @@
+// Copyright Epic Games, Inc. All Rights Reserved.
+
+#include "ModelingToolsHostCustomizationAPI.h"
+
+#include "ContextObjectStore.h"
+#include "InteractiveToolsContext.h"
+#include "ModelingToolsEditorModeToolkit.h"
+
+bool UModelingToolsHostCustomizationAPI::RequestAcceptCancelButtonOverride(FAcceptCancelButtonOverrideParams& Params)
+{
+	if (TSharedPtr<FModelingToolsEditorModeToolkit> ToolkitPinned = Toolkit.Pin())
+	{
+		return ToolkitPinned->RequestAcceptCancelButtonOverride(Params);
+	}
+	return false;
+}
+
+bool UModelingToolsHostCustomizationAPI::RequestCompleteButtonOverride(FCompleteButtonOverrideParams& Params)
+{
+	if (TSharedPtr<FModelingToolsEditorModeToolkit> ToolkitPinned = Toolkit.Pin())
+	{
+		return ToolkitPinned->RequestCompleteButtonOverride(Params);
+	}
+	return false;
+}
+
+void UModelingToolsHostCustomizationAPI::ClearButtonOverrides()
+{
+	if (TSharedPtr<FModelingToolsEditorModeToolkit> ToolkitPinned = Toolkit.Pin())
+	{
+		ToolkitPinned->ClearButtonOverrides();
+	}
+}
+
+UModelingToolsHostCustomizationAPI* UModelingToolsHostCustomizationAPI::Register(
+	UInteractiveToolsContext* ToolsContext, TSharedRef<FModelingToolsEditorModeToolkit> Toolkit)
+{
+	if (!ensure(ToolsContext && ToolsContext->ContextObjectStore))
+	{
+		return nullptr;
+	}
+
+	UModelingToolsHostCustomizationAPI* Found = ToolsContext->ContextObjectStore->FindContext<UModelingToolsHostCustomizationAPI>();
+	if (Found)
+	{
+		if (!ensureMsgf(Found->Toolkit == Toolkit,
+			TEXT("UModelingToolsHostCustomizationAPI already registered, but with different toolkit. "
+				"Do not expect multiple toolkits to provide tool overlays to customize.")))
+		{
+			Found->Toolkit = Toolkit;
+		}
+		return Found;
+	}
+	else
+	{
+		UModelingToolsHostCustomizationAPI* Instance = NewObject<UModelingToolsHostCustomizationAPI>();
+		Instance->Toolkit = Toolkit;
+		ensure(ToolsContext->ContextObjectStore->AddContextObject(Instance));
+		return Instance;
+	}
+
+	ensure(false); // Shouldn't get here
+	return nullptr;
+}
+
+bool UModelingToolsHostCustomizationAPI::Deregister(UInteractiveToolsContext* ToolsContext)
+{
+	if (!ensure(ToolsContext && ToolsContext->ContextObjectStore))
+	{
+		return false;
+	}
+
+	ToolsContext->ContextObjectStore->RemoveContextObjectsOfType(UModelingToolsHostCustomizationAPI::StaticClass());
+	return true;
+}
diff --git a/Engine/Plugins/Editor/ModelingToolsEditorMode/Source/ModelingToolsEditorMode/Public/ModelingToolsEditorModeToolkit.h b/Engine/Plugins/Editor/ModelingToolsEditorMode/Source/ModelingToolsEditorMode/Public/ModelingToolsEditorModeToolkit.h
--- a/Engine/Plugins/Editor/ModelingToolsEditorMode/Source/ModelingToolsEditorMode/Public/ModelingToolsEditorModeToolkit.h
+++ b/Engine/Plugins/Editor/ModelingToolsEditorMode/Source/ModelingToolsEditorMode/Public/ModelingToolsEditorModeToolkit.h
@@ -4,6 +4,7 @@
 
 #include "Toolkits/BaseToolkit.h"
 #include "StatusBarSubsystem.h"
+#include "ToolHostCustomizationAPI.h" // FAcceptCancelButtonOverrideParams
 
 class IAssetViewport;
 class SEditableTextBox;
@@ -79,6 +80,10 @@ public:
 	virtual void InvokeUI() override;
 	virtual void ExtendSecondaryModeToolbar(UToolMenu *InModeToolbarMenu) override;
 
+	virtual bool RequestAcceptCancelButtonOverride(IToolHostCustomizationAPI::FAcceptCancelButtonOverrideParams& Params);
+	virtual bool RequestCompleteButtonOverride(IToolHostCustomizationAPI::FCompleteButtonOverrideParams& Params);
+	virtual void ClearButtonOverrides();
+
 	// Read and write Autogen Subfolder quick setting in restrictive mode.
 	virtual FText GetRestrictiveModeAutoGeneratedAssetPathText() const;
 	virtual void OnRestrictiveModeAutoGeneratedAssetPathTextCommitted(const FText& InNewText, ETextCommit::Type InTextCommit) const;
@@ -110,6 +115,13 @@ private:
 	FStatusBarMessageHandle ActiveToolMessageHandle;
 	const FSlateBrush* ActiveToolIcon = nullptr;
 
+	TOptional<IToolHostCustomizationAPI::FAcceptCancelButtonOverrideParams> AcceptCancelButtonParams;
+	TOptional<IToolHostCustomizationAPI::FCompleteButtonOverrideParams> CompleteButtonParams;
+	bool bCurrentOverrideButtonsWereClicked = false;
+
+	FReply HandleAcceptCancelClick(bool bAccept);
+	FReply HandleCompleteClick();
+
 	UGeometrySelectionManager* CachedSelectionManager = nullptr;
 
 	TSharedPtr<SWidget> ToolkitWidget;
@@ -198,6 +210,7 @@ private:
 
 	bool bFirstInitializeAfterModeSetup = true;
 	bool bShowActiveSelectionActions = true;
+
 };
 
 #undef LOCTEXT_NAMESPACE
diff --git a/Engine/Plugins/Editor/ModelingToolsEditorMode/Source/ModelingToolsEditorMode/Public/ModelingToolsHostCustomizationAPI.h b/Engine/Plugins/Editor/ModelingToolsEditorMode/Source/ModelingToolsEditorMode/Public/ModelingToolsHostCustomizationAPI.h
new file mode 100644
--- /dev/null
+++ b/Engine/Plugins/Editor/ModelingToolsEditorMode/Source/ModelingToolsEditorMode/Public/ModelingToolsHostCustomizationAPI.h
@@ -0,0 +1,42 @@
+// Copyright Epic Games, Inc. All Rights Reserved.
+
+#pragma once
+
+#include "Templates/SharedPointer.h" // TWeakPtr, TSharedRef
+#include "ToolHostCustomizationAPI.h"
+#include "UObject/Object.h"
+
+#include "ModelingToolsHostCustomizationAPI.generated.h"
+
+class FModelingToolsEditorModeToolkit;
+class UInteractiveToolsContext;
+
+/**
+ * An implementation of IToolHostCustomizationAPI meant to be used with modeling tools editor mode.
+ * Mostly it just forwards the relevant requests to the mode toolkit, which changes its accept/cancel/complete
+ * overlay buttons temporarily (until the next click of those buttons).
+ */
+UCLASS(Transient, MinimalAPI)
+class UModelingToolsHostCustomizationAPI : public UObject, public IToolHostCustomizationAPI
+{
+	GENERATED_BODY()
+
+public:
+	virtual bool SupportsAcceptCancelButtonOverride() override { return true; }
+	virtual bool RequestAcceptCancelButtonOverride(FAcceptCancelButtonOverrideParams& Params) override;
+
+	virtual bool SupportsCompleteButtonOverride() override { return true; }
+	virtual bool RequestCompleteButtonOverride(FCompleteButtonOverrideParams& Params) override;
+
+	virtual void ClearButtonOverrides() override;
+
+	/** Helper to create, initialize, and register a UModelingToolsHostCustomizationAPI object with the context store */
+	static UModelingToolsHostCustomizationAPI* Register(UInteractiveToolsContext* ToolsContext, 
+		TSharedRef<FModelingToolsEditorModeToolkit> Toolkit);
+
+	/** Helper to deregister the UModelingToolsHostCustomizationAPI from the context store */
+	static bool Deregister(UInteractiveToolsContext* ToolsContext);
+
+protected:
+	TWeakPtr<FModelingToolsEditorModeToolkit> Toolkit;
+};
\ No newline at end of file
diff --git a/Engine/Plugins/Runtime/MeshModelingToolset/Source/MeshModelingTools/Private/EditMeshPolygonsTool.cpp b/Engine/Plugins/Runtime/MeshModelingToolset/Source/MeshModelingTools/Private/EditMeshPolygonsTool.cpp
--- a/Engine/Plugins/Runtime/MeshModelingToolset/Source/MeshModelingTools/Private/EditMeshPolygonsTool.cpp
+++ b/Engine/Plugins/Runtime/MeshModelingToolset/Source/MeshModelingTools/Private/EditMeshPolygonsTool.cpp
@@ -18,6 +18,7 @@
 #include "DynamicMeshEditor.h"
 #include "FaceGroupUtil.h"
 #include "GroupTopology.h"
+#include "Input/Reply.h"
 #include "InteractiveGizmoManager.h"
 #include "InteractiveToolManager.h"
 #include "Mechanics/DragAlignmentMechanic.h"
@@ -33,6 +34,7 @@
 #include "Selection/StoredMeshSelectionUtil.h"
 #include "Selection/PolygonSelectionMechanic.h"
 #include "Selections/MeshConnectedComponents.h"
+#include "Styling/AppStyle.h"
 #include "TargetInterfaces/MaterialProvider.h"
 #include "TargetInterfaces/MeshDescriptionCommitter.h"
 #include "TargetInterfaces/MeshDescriptionProvider.h"
@@ -46,6 +48,7 @@
 #include "ToolActivities/PolyEditPlanarProjectionUVActivity.h"
 #include "ToolActivities/PolyEditBevelEdgeActivity.h"
 #include "ToolContextInterfaces.h" // FToolBuilderState
+#include "ToolHostCustomizationAPI.h"
 #include "ToolSetupUtil.h"
 #include "ToolTargetManager.h"
 #include "TransformTypes.h"
@@ -441,17 +444,6 @@ void UEditMeshPolygonsTool::Setup()
 	// We add an empty line for the error message so that things don't jump when we use it.
 	GetToolManager()->DisplayMessage(FText(), EToolMessageLevel::UserWarning);
 
-	CancelAction = NewObject<UEditMeshPolygonsToolCancelAction>();
-	CancelAction->Initialize(this);
-	AddToolPropertySource(CancelAction);
-	SetToolPropertySourceEnabled(CancelAction, false);
-
-	AcceptCancelAction = NewObject<UEditMeshPolygonsToolAcceptCancelAction>();
-	AcceptCancelAction->Initialize(this);
-	AddToolPropertySource(AcceptCancelAction);
-	SetToolPropertySourceEnabled(AcceptCancelAction, false);
-
-
 	// Initialize the common properties but don't add them yet, because we want them to be under the activity-specific ones.
 	CommonProps = NewObject<UPolyEditCommonProperties>(this);
 	CommonProps->RestoreProperties(this, GetPropertyCacheIdentifier(bTriangleMode));
@@ -585,24 +577,38 @@ void UEditMeshPolygonsTool::Setup()
 
 	ExtrudeActivity = NewObject<UPolyEditExtrudeActivity>();
 	ExtrudeActivity->Setup(this);
+	ActivityLabels.Add(ExtrudeActivity, LOCTEXT("ExtrudeActivityLabel", "Extrude"));
+	ActivityIconNames.Add(ExtrudeActivity, "PolyEd.Extrude");
 	
 	InsetOutsetActivity = NewObject<UPolyEditInsetOutsetActivity>();
 	InsetOutsetActivity->Setup(this);
+	// The icons/labels differ depending on whether we are doing an inset or outset, so we set those 
+	// when we launch the activity.
 
 	CutFacesActivity = NewObject<UPolyEditCutFacesActivity>();
 	CutFacesActivity->Setup(this);
+	ActivityLabels.Add(CutFacesActivity, LOCTEXT("CutFacesActivityLabel", "Cut Faces"));
+	ActivityIconNames.Add(CutFacesActivity, "PolyEd.CutFaces");
 
 	PlanarProjectionUVActivity = NewObject<UPolyEditPlanarProjectionUVActivity>();
 	PlanarProjectionUVActivity->Setup(this);
+	ActivityLabels.Add(PlanarProjectionUVActivity, LOCTEXT("UVProjectActivityLabel", "UV Project"));
+	ActivityIconNames.Add(PlanarProjectionUVActivity, "PolyEd.ProjectUVs");
 
 	InsertEdgeLoopActivity = NewObject<UPolyEditInsertEdgeLoopActivity>();
 	InsertEdgeLoopActivity->Setup(this);
+	ActivityLabels.Add(InsertEdgeLoopActivity, LOCTEXT("InsertEdgeLoopsActivityLabel", "Insert Edge Loops"));
+	ActivityIconNames.Add(InsertEdgeLoopActivity, "PolyEd.InsertEdgeLoop");
 
 	InsertEdgeActivity = NewObject<UPolyEditInsertEdgeActivity>();
 	InsertEdgeActivity->Setup(this);
+	ActivityLabels.Add(InsertEdgeActivity, LOCTEXT("InsertEdgesActivityLabel", "Insert Edges"));
+	ActivityIconNames.Add(InsertEdgeActivity, "PolyEd.InsertGroupEdge");
 
 	BevelEdgeActivity = NewObject<UPolyEditBevelEdgeActivity>();
 	BevelEdgeActivity->Setup(this);
+	ActivityLabels.Add(BevelEdgeActivity, LOCTEXT("BevelActivityLabel", "Bevel"));
+	ActivityIconNames.Add(BevelEdgeActivity, "PolyEd.Bevel");
 
 	// Now that we've initialized the activities, add in the selection settings and 
 	// CommonProps so that they are at the bottom.
@@ -707,6 +713,10 @@ void UEditMeshPolygonsTool::OnShutdown(EToolShutdownType ShutdownType)
 
 	if (CurrentActivity)
 	{
+		if (IToolHostCustomizationAPI* ButtonCustomizer = IToolHostCustomizationAPI::Find(GetToolManager()).GetInterface())
+		{
+			ButtonCustomizer->ClearButtonOverrides();
+		}
 		CurrentActivity->End(ShutdownType);
 		CurrentActivity = nullptr;
 	}
@@ -805,8 +815,6 @@ void UEditMeshPolygonsTool::OnShutdown(EToolShutdownType ShutdownType)
 	EditEdgeActions = nullptr;
 	EditEdgeActions_Triangles = nullptr;
 	EditUVActions = nullptr;
-	CancelAction = nullptr;
-	AcceptCancelAction = nullptr;
 
 	ExtrudeActivity = nullptr;
 	InsetOutsetActivity = nullptr;
@@ -1218,12 +1226,20 @@ void UEditMeshPolygonsTool::OnTick(float DeltaTime)
 		case EEditMeshPolygonsToolActions::Inset:
 		{
 			InsetOutsetActivity->Settings->bOutset = false;
+			
+			ActivityLabels.Add(InsetOutsetActivity, LOCTEXT("InsetActivityLabel", "Inset"));
+			ActivityIconNames.Add(InsetOutsetActivity, "PolyEd.Inset");
+			
 			StartActivity(InsetOutsetActivity);
 			break;
 		}
 		case EEditMeshPolygonsToolActions::Outset:
 		{
 			InsetOutsetActivity->Settings->bOutset = true;
+			
+			ActivityLabels.Add(InsetOutsetActivity, LOCTEXT("OutsetActivityLabel", "Outset"));
+			ActivityIconNames.Add(InsetOutsetActivity, "PolyEd.Outset");
+
 			StartActivity(InsetOutsetActivity);
 			break;
 		}
@@ -1255,17 +1271,6 @@ void UEditMeshPolygonsTool::OnTick(float DeltaTime)
 			break;
 		}
 
-		case EEditMeshPolygonsToolActions::CancelCurrent:
-		{
-			EndCurrentActivity(EToolShutdownType::Cancel);
-			break;
-		}
-		case EEditMeshPolygonsToolActions::AcceptCurrent:
-		{
-			EndCurrentActivity(EToolShutdownType::Accept);
-			break;
-		}
-
 		// Single action operations:
 		case EEditMeshPolygonsToolActions::Merge:
 			ApplyMerge();
@@ -1353,13 +1358,49 @@ void UEditMeshPolygonsTool::StartActivity(TObjectPtr<UInteractiveToolActivity> A
 
 		if ( bTerminateOnPendingActionComplete == false )
 		{
-			if (CurrentActivity->HasAccept())
+			// Customize the tool accept/cancel buttons to the current activity.
+			if (IToolHostCustomizationAPI* ButtonCustomizer = IToolHostCustomizationAPI::Find(GetToolManager()).GetInterface())
 			{
-				SetToolPropertySourceEnabled(AcceptCancelAction, true);
-			}
-			else
-			{
-				SetToolPropertySourceEnabled(CancelAction, true);
+				const FText SubActionFallbackLabel = LOCTEXT("SubActionFallbackLabel", "Current Action");
+				if (CurrentActivity->HasAccept())
+				{
+					IToolHostCustomizationAPI::FAcceptCancelButtonOverrideParams Params;
+					Params.Label = ActivityLabels.Contains(CurrentActivity) ? ActivityLabels[CurrentActivity] : SubActionFallbackLabel;
+					if (ActivityIconNames.Contains(CurrentActivity))
+					{
+						Params.IconName = ActivityIconNames[CurrentActivity];
+					}
+					Params.OverrideAcceptButtonText = LOCTEXT("AcceptSubActionButton", "Accept Action");
+					Params.OverrideAcceptButtonTooltip = LOCTEXT("AcceptSubActionTooltip", "Accept the action currently being performed.");
+					Params.OverrideCancelButtonText = LOCTEXT("CancelSubActionButton", "Cancel Action");
+					Params.OverrideCancelButtonTooltip = LOCTEXT("AcceptSubActionTooltip", "Cancel the action currently being performed.");
+					Params.CanAccept = [this]() { return CurrentActivity->CanAccept(); };
+					Params.OnAcceptCancelTriggered = [this](bool bAccept) 
+					{ 
+						EndCurrentActivity(bAccept ? EToolShutdownType::Accept : EToolShutdownType::Cancel);
+						return FReply::Handled();
+					};
+
+					ButtonCustomizer->RequestAcceptCancelButtonOverride(Params);
+				}
+				else
+				{
+					IToolHostCustomizationAPI::FCompleteButtonOverrideParams Params;
+					Params.Label = ActivityLabels.Contains(CurrentActivity) ? ActivityLabels[CurrentActivity] : SubActionFallbackLabel;
+					if (ActivityIconNames.Contains(CurrentActivity))
+					{
+						Params.IconName = ActivityIconNames[CurrentActivity];
+					}
+					Params.OverrideCompleteButtonText = LOCTEXT("CompleteSubActionButton", "Done");
+					Params.OverrideCompleteButtonTooltip = LOCTEXT("CompleteSubActionTooltip", "Exit the current activity.");
+					Params.OnCompleteTriggered = [this]() 
+					{
+						EndCurrentActivity(EToolShutdownType::Completed);
+						return FReply::Handled();
+					};
+
+					ButtonCustomizer->RequestCompleteButtonOverride(Params);
+				}
 			}
 		}
 		else
@@ -1397,8 +1438,10 @@ void UEditMeshPolygonsTool::EndCurrentActivity(EToolShutdownType ShutdownType)
 			return;
 		}
 
-		SetToolPropertySourceEnabled(CancelAction, false);
-		SetToolPropertySourceEnabled(AcceptCancelAction, false);
+		if (IToolHostCustomizationAPI* ButtonCustomizer = IToolHostCustomizationAPI::Find(GetToolManager()).GetInterface())
+		{
+			ButtonCustomizer->ClearButtonOverrides();
+		}
 		SetActionButtonPanelsVisible(true);
 		SelectionMechanic->SetIsEnabled(true);
 		SetToolPropertySourceEnabled(TopologyProperties, true);
diff --git a/Engine/Plugins/Runtime/MeshModelingToolset/Source/MeshModelingTools/Public/EditMeshPolygonsTool.h b/Engine/Plugins/Runtime/MeshModelingToolset/Source/MeshModelingTools/Public/EditMeshPolygonsTool.h
--- a/Engine/Plugins/Runtime/MeshModelingToolset/Source/MeshModelingTools/Public/EditMeshPolygonsTool.h
+++ b/Engine/Plugins/Runtime/MeshModelingToolset/Source/MeshModelingTools/Public/EditMeshPolygonsTool.h
@@ -10,6 +10,7 @@
 #include "InteractiveToolActivity.h" // IToolActivityHost
 #include "InteractiveToolBuilder.h"
 #include "InteractiveToolQueryInterfaces.h" // IInteractiveToolNestedAcceptCancelAPI
+#include "Internationalization/Text.h"
 #include "Operations/GroupTopologyDeformer.h"
 #include "BaseTools/SingleTargetWithSelectionTool.h"
 
@@ -20,6 +21,7 @@
 PREDECLARE_GEOMETRY(class FGroupTopology);
 PREDECLARE_GEOMETRY(struct FGroupTopologySelection);
 
+struct FSlateBrush;
 class UCombinedTransformGizmo;
 class UDragAlignmentMechanic;
 class UMeshOpPreviewWithBackgroundCompute; 
@@ -452,40 +454,6 @@ public:
 };
 
 
-/**
- * TODO: This is currently a separate action set so that we can show/hide it depending on whether
- * we have an activity running. We should have a cleaner alternative.
- */
-UCLASS()
-class MESHMODELINGTOOLS_API UEditMeshPolygonsToolCancelAction : public UEditMeshPolygonsToolActionPropertySet
-{
-	GENERATED_BODY()
-public:
-	// Complete the current action
-	//~ Note this function may accept or cancel depending on the tool, despite the class name and the PostAction argument
-	UFUNCTION(CallInEditor, Category = CurrentOperation, meta = (DisplayName = "Complete", DisplayPriority = 1))
-	void Done() { PostAction(EEditMeshPolygonsToolActions::CancelCurrent); }
-};
-
-
-/**
- * TODO: This is currently a separate action set so that we can show/hide it depending on whether
- * we have an activity running. We should have a cleaner alternative.
- */
-UCLASS()
-class MESHMODELINGTOOLS_API UEditMeshPolygonsToolAcceptCancelAction : public UEditMeshPolygonsToolActionPropertySet
-{
-	GENERATED_BODY()
-public:
-	UFUNCTION(CallInEditor, Category = CurrentOperation, meta = (DisplayName = "Apply", DisplayPriority = 1))
-	void Apply() { PostAction(EEditMeshPolygonsToolActions::AcceptCurrent); }
-
-	UFUNCTION(CallInEditor, Category = CurrentOperation, meta = (DisplayName = "Cancel", DisplayPriority = 2))
-	void Cancel() { PostAction(EEditMeshPolygonsToolActions::CancelCurrent); }
-};
-
-
-
 
 /**
  *
@@ -576,12 +544,6 @@ protected:
 	UPROPERTY()
 	TObjectPtr<UEditMeshPolygonsToolUVActions> EditUVActions = nullptr;
 
-	UPROPERTY()
-	TObjectPtr<UEditMeshPolygonsToolCancelAction> CancelAction = nullptr;
-
-	UPROPERTY()
-	TObjectPtr<UEditMeshPolygonsToolAcceptCancelAction> AcceptCancelAction = nullptr;
-
 	UPROPERTY()
 	TObjectPtr<UPolyEditTopologyProperties> TopologyProperties = nullptr;
 
@@ -603,6 +565,9 @@ protected:
 	UPROPERTY()
 	TObjectPtr<UPolyEditBevelEdgeActivity> BevelEdgeActivity = nullptr;
 
+	TMap<UInteractiveToolActivity*, FText> ActivityLabels;
+	TMap<UInteractiveToolActivity*, FName> ActivityIconNames;
+
 	/**
 	 * Points to one of the activities when it is active
 	 */
diff --git a/Engine/Plugins/Runtime/MeshModelingToolset/Source/ModelingComponents/Private/ToolHostCustomizationAPI.cpp b/Engine/Plugins/Runtime/MeshModelingToolset/Source/ModelingComponents/Private/ToolHostCustomizationAPI.cpp
new file mode 100644
--- /dev/null
+++ b/Engine/Plugins/Runtime/MeshModelingToolset/Source/ModelingComponents/Private/ToolHostCustomizationAPI.cpp
@@ -0,0 +1,23 @@
+// Copyright Epic Games, Inc. All Rights Reserved.
+
+#include "ToolHostCustomizationAPI.h"
+
+#include "ContextObjectStore.h"
+#include "InteractiveToolManager.h"
+#include "UObject/Object.h"
+
+TScriptInterface<IToolHostCustomizationAPI> IToolHostCustomizationAPI::Find(UInteractiveToolManager* ToolManager)
+{
+	if (!ensure(ToolManager))
+	{
+		return nullptr;
+	}
+
+	UContextObjectStore* ContextObjectStore = ToolManager->GetContextObjectStore();
+	if (!ensure(ContextObjectStore))
+	{
+		return nullptr;
+	}
+
+	return ContextObjectStore->FindContextByClass(UToolHostCustomizationAPI::StaticClass());
+}
\ No newline at end of file
diff --git a/Engine/Plugins/Runtime/MeshModelingToolset/Source/ModelingComponents/Public/ToolHostCustomizationAPI.h b/Engine/Plugins/Runtime/MeshModelingToolset/Source/ModelingComponents/Public/ToolHostCustomizationAPI.h
new file mode 100644
--- /dev/null
+++ b/Engine/Plugins/Runtime/MeshModelingToolset/Source/ModelingComponents/Public/ToolHostCustomizationAPI.h
@@ -0,0 +1,101 @@
+// Copyright Epic Games, Inc. All Rights Reserved.
+
+#pragma once
+
+#include "Internationalization/Text.h"
+#include "Templates/Function.h"
+#include "UObject/Interface.h"
+#include "UObject/ScriptInterface.h"
+
+#include "ToolHostCustomizationAPI.generated.h"
+
+class UInteractiveToolManager;
+struct FSlateBrush;
+
+UINTERFACE(MinimalAPI)
+class UToolHostCustomizationAPI : public UInterface
+{
+	GENERATED_BODY()
+};
+
+/**
+ * An API (to be stored as a context object) that would allow tools to customize aspects of their
+ * hosts' appearance. For instance, it can allow a tool to customize tool shutdown accept/cancel
+ * buttons to handle some tool sub-action, if the host allows this.
+ */
+class IToolHostCustomizationAPI
+{
+	GENERATED_BODY()
+public:
+
+	/** If true, the tool host is able to override its tool shutdown buttons to custom accept/cancel buttons. */
+	virtual bool SupportsAcceptCancelButtonOverride() { return false; }
+
+	struct FAcceptCancelButtonOverrideParams
+	{
+		/** Function to be called when the accept/cancel buttons are clicked. Passed in arg is true if it was the accept button. */
+		TFunction<void(bool /*bAccept*/)> OnAcceptCancelTriggered;
+		
+		/** If this returns false, the accept button is not clickable. */
+		TFunction<bool()> CanAccept = []() { return true; };
+		
+		/** Label to use to show that the buttons now control a sub action. */
+		FText Label;
+
+		// Optional parameters:
+
+		/** The host will try to look up an icon with this name in its style sheet. Icon will not be used if host does not find a match. */
+		TOptional<FName> IconName;
+
+		TOptional<FText> OverrideAcceptButtonText;
+		TOptional<FText> OverrideAcceptButtonTooltip;
+		TOptional<FText> OverrideCancelButtonText;
+		TOptional<FText> OverrideCancelButtonTooltip;
+	};
+
+	/**
+	 * Requests that a tool host's tool shutdown buttons be temporarily overriden to an accept/cancel pair
+	 * that accepts/cancels some sub action.
+	 */
+	virtual bool RequestAcceptCancelButtonOverride(FAcceptCancelButtonOverrideParams& Params) { return false; }
+
+	virtual bool SupportsCompleteButtonOverride() { return false; }
+
+	struct FCompleteButtonOverrideParams
+	{
+		/** Function to be called when complete button is clicked. */
+		TFunction<void()> OnCompleteTriggered;
+
+		/** Label to use to show that the buttons now control a sub action. */
+		FText Label;
+		
+		// Optional parameters:
+
+		/** The host will try to look up an icon with this name in its style sheet. Icon will not be used if host does not find a match. */
+		TOptional<FName> IconName;
+
+		TOptional<FText> OverrideCompleteButtonText;
+		TOptional<FText> OverrideCompleteButtonTooltip;
+	};
+
+	/**
+	 * Requests that a tool host's shutdown buttons be temporarily overriden to a complete button that
+	 * completes some sub action.
+	 */
+	virtual bool RequestCompleteButtonOverride(FCompleteButtonOverrideParams& Params) { return false; }
+
+	/**
+	 * Remove any requested button customizations. 
+	 * 
+	 * This can be imporant to call on tool shutdown if the tool registered button overrides
+	 * whose handlers hold a reference to the tool, or else, for example, a CanAccept override
+	 * could end up trying to query the tool after shutdown.
+	 */
+	virtual void ClearButtonOverrides() = 0;
+
+public:
+	/**
+	 * @return existing IToolHostCustomizationAPI registered in context store via the ToolManager, or nullptr if not found
+	 */
+	static MODELINGCOMPONENTS_API TScriptInterface<IToolHostCustomizationAPI> Find(UInteractiveToolManager* ToolManager);
+};
\ No newline at end of file
