Modeling Tools: Adds 'Local Frame Mode' toggle in the Mesh Element Selection Toolbar Settings so users can edit with the gizmo based either on the selected geometry only or on the entire object

diff --git a/Engine/Plugins/Editor/ModelingToolsEditorMode/Source/ModelingToolsEditorMode/Private/ModelingModeToolkit_Toolbars.cpp b/Engine/Plugins/Editor/ModelingToolsEditorMode/Source/ModelingToolsEditorMode/Private/ModelingModeToolkit_Toolbars.cpp
--- a/Engine/Plugins/Editor/ModelingToolsEditorMode/Source/ModelingToolsEditorMode/Private/ModelingModeToolkit_Toolbars.cpp
+++ b/Engine/Plugins/Editor/ModelingToolsEditorMode/Source/ModelingToolsEditorMode/Private/ModelingModeToolkit_Toolbars.cpp
@@ -257,7 +257,7 @@ void MakeSubMenu_DefaultMeshObjectType(FMenuBuilder& MenuBuilder)
 {
 	UModelingToolsEditorModeSettings* Settings = GetMutableDefault<UModelingToolsEditorModeSettings>();
 
-	MenuBuilder.BeginSection("DefaultObjecTypeSection", LOCTEXT("DefaultObjecTypeSection", "Default Object Type"));
+	MenuBuilder.BeginSection("DefaultObjectTypeSection", LOCTEXT("DefaultObjectTypeSection", "Default Object Type"));
 
 	auto ShowMeshObjectDefaultChangeToast = []()
 	{
@@ -415,6 +415,38 @@ void MakeSubMenu_Selection_MeshType(FModelingToolsEditorModeToolkit* Toolkit, FM
 		FSlateIcon(), ToggleStaticMeshesAction, NAME_None, EUserInterfaceActionType::ToggleButton);
 
 }
+	
+void MakeSubMenu_Selection_LocalFrameMode(FModelingToolsEditorModeToolkit* Toolkit, FMenuBuilder& MenuBuilder)
+{
+		
+	UModelingSelectionInteraction* SelectionInteraction = Cast<UModelingToolsEditorMode>(Toolkit->GetScriptableEditorMode())->GetSelectionInteraction();
+	auto ToggleFromGeometryAction = [SelectionInteraction](EModelingSelectionInteraction_LocalFrameMode LocalFrameMode)
+	{
+		return FUIAction(
+			FExecuteAction::CreateLambda([SelectionInteraction, LocalFrameMode]
+			{
+				SelectionInteraction->SetLocalFrameMode(LocalFrameMode);
+				UModelingToolsModeCustomizationSettings* ModelingEditorSettings = GetMutableDefault<UModelingToolsModeCustomizationSettings>();
+				ModelingEditorSettings->LastMeshSelectionLocalFrameMode = static_cast<int>(LocalFrameMode);
+				ModelingEditorSettings->SaveConfig();
+			}),
+			FCanExecuteAction(),
+			FIsActionChecked::CreateLambda([SelectionInteraction, LocalFrameMode]()
+			{
+				return (SelectionInteraction->GetLocalFrameMode() == LocalFrameMode);
+			}));
+	};
+	
+	MenuBuilder.AddMenuEntry(
+		LOCTEXT("Selection_LocalFrameMode_FromGeometry", "From Geometry"), 
+		LOCTEXT("Selection_LocalFrameMode_FromGeometry_Tooltip", "Gizmo Orientation Based on Selected Geometry"),
+		FSlateIcon(), ToggleFromGeometryAction(EModelingSelectionInteraction_LocalFrameMode::FromGeometry), NAME_None, EUserInterfaceActionType::ToggleButton);
+
+	MenuBuilder.AddMenuEntry(
+		LOCTEXT("Selection_LocalFrameMode_FromObject", "From Object"), 
+		LOCTEXT("Selection_LocalFrameMode_FromObject_Tooltip", "Gizmo Orientation Based on Object"),
+		FSlateIcon(), ToggleFromGeometryAction(EModelingSelectionInteraction_LocalFrameMode::FromObject), NAME_None, EUserInterfaceActionType::ToggleButton);
+}
 
 TSharedRef<SWidget> MakeMenu_SelectionConfigSettings(FModelingToolsEditorModeToolkit* Toolkit)
 {
@@ -428,6 +460,10 @@ TSharedRef<SWidget> MakeMenu_SelectionConfigSettings(FModelingToolsEditorModeToo
 	MakeSubMenu_Selection_MeshType(Toolkit, MenuBuilder);
 	MenuBuilder.EndSection();
 
+	MenuBuilder.BeginSection("Section_LocalFrameMode", LOCTEXT("Section_LocalFrameMode", "Local Frame Mode"));
+	MakeSubMenu_Selection_LocalFrameMode(Toolkit, MenuBuilder);
+	MenuBuilder.EndSection();
+
 	TSharedRef<SWidget> MenuWidget = MenuBuilder.MakeWidget();
 	return MenuWidget;
 }
diff --git a/Engine/Plugins/Editor/ModelingToolsEditorMode/Source/ModelingToolsEditorMode/Private/ModelingSelectionInteraction.cpp b/Engine/Plugins/Editor/ModelingToolsEditorMode/Source/ModelingToolsEditorMode/Private/ModelingSelectionInteraction.cpp
--- a/Engine/Plugins/Editor/ModelingToolsEditorMode/Source/ModelingToolsEditorMode/Private/ModelingSelectionInteraction.cpp
+++ b/Engine/Plugins/Editor/ModelingToolsEditorMode/Source/ModelingToolsEditorMode/Private/ModelingSelectionInteraction.cpp
@@ -120,6 +120,15 @@ void UModelingSelectionInteraction::SetActiveDragMode(EModelingSelectionInteract
 	}
 }
 
+void UModelingSelectionInteraction::SetLocalFrameMode(EModelingSelectionInteraction_LocalFrameMode NewLocalFrameMode)
+{
+	if (LocalFrameMode != NewLocalFrameMode)
+	{
+		LocalFrameMode = NewLocalFrameMode;
+		UpdateGizmoOnSelectionChange();
+	}
+}
+
 void UModelingSelectionInteraction::UpdateActiveDragMode()
 {
 	if (ActiveDragMode == EModelingSelectionInteraction_DragMode::PathInteraction)
@@ -339,7 +348,14 @@ void UModelingSelectionInteraction::UpdateGizmoOnSelectionChange()
 		TransformGizmo->SetVisibility(true);
 
 		FFrame3d SelectionFrame;
-		SelectionManager->GetSelectionWorldFrame(SelectionFrame);
+		if (LocalFrameMode == EModelingSelectionInteraction_LocalFrameMode::FromGeometry)
+		{
+			SelectionManager->GetSelectionWorldFrame(SelectionFrame);
+		}
+		else
+		{
+			SelectionManager->GetTargetWorldFrame(SelectionFrame);
+		}
 		TransformGizmo->ReinitializeGizmoTransform( SelectionFrame.ToFTransform() );
 	}
 }
diff --git a/Engine/Plugins/Editor/ModelingToolsEditorMode/Source/ModelingToolsEditorMode/Private/ModelingToolsEditorMode.cpp b/Engine/Plugins/Editor/ModelingToolsEditorMode/Source/ModelingToolsEditorMode/Private/ModelingToolsEditorMode.cpp
--- a/Engine/Plugins/Editor/ModelingToolsEditorMode/Source/ModelingToolsEditorMode/Private/ModelingToolsEditorMode.cpp
+++ b/Engine/Plugins/Editor/ModelingToolsEditorMode/Source/ModelingToolsEditorMode/Private/ModelingToolsEditorMode.cpp
@@ -1074,6 +1074,13 @@ void UModelingToolsEditorMode::Enter()
 		{
 			SelectionInteraction->SetActiveDragMode(LastDragMode);
 		}
+
+		EModelingSelectionInteraction_LocalFrameMode LastLocalFrameMode =
+			static_cast<EModelingSelectionInteraction_LocalFrameMode>(ModelingEditorSettings->LastMeshSelectionLocalFrameMode);
+		if ( LastLocalFrameMode == EModelingSelectionInteraction_LocalFrameMode::FromGeometry || LastLocalFrameMode == EModelingSelectionInteraction_LocalFrameMode::FromObject )
+		{
+			SelectionInteraction->SetLocalFrameMode(LastLocalFrameMode);
+		}
 	}
 
 	// initialize SelectionManager w/ active selection
diff --git a/Engine/Plugins/Editor/ModelingToolsEditorMode/Source/ModelingToolsEditorMode/Public/ModelingSelectionInteraction.h b/Engine/Plugins/Editor/ModelingToolsEditorMode/Source/ModelingToolsEditorMode/Public/ModelingSelectionInteraction.h
--- a/Engine/Plugins/Editor/ModelingToolsEditorMode/Source/ModelingToolsEditorMode/Public/ModelingSelectionInteraction.h
+++ b/Engine/Plugins/Editor/ModelingToolsEditorMode/Source/ModelingToolsEditorMode/Public/ModelingSelectionInteraction.h
@@ -35,6 +35,13 @@ enum class EModelingSelectionInteraction_DragMode : uint8
 	RectangleMarqueeInteraction = 2
 };
 
+UENUM()
+enum class EModelingSelectionInteraction_LocalFrameMode : uint8
+{
+	FromGeometry = 0,
+	FromObject = 1,
+};
+
 /**
  * UModelingSelectionInteraction provides element-level selection behavior (ie mesh triangles/edges/vertices)
  * via the UGeometrySelectionManager. The Interaction doesn't actually know about the meshes or interact
@@ -98,6 +105,11 @@ public:
 	DECLARE_MULTICAST_DELEGATE(FOnEndTransformInteraction);
 	FOnEndTransformInteraction OnTransformEnd;
 
+	//
+	// Support for Adjusting Local Frame Mode
+	//
+	EModelingSelectionInteraction_LocalFrameMode GetLocalFrameMode() const { return LocalFrameMode; }
+	void SetLocalFrameMode(EModelingSelectionInteraction_LocalFrameMode NewLocalFrameMode);
 
 protected:
 	// click-to-select behavior, various drag behaviors
@@ -207,6 +219,8 @@ protected:
 	EModelingSelectionInteraction_DragMode ActiveDragMode = EModelingSelectionInteraction_DragMode::NoDragInteraction;
 	void UpdateActiveDragMode();
 
+	EModelingSelectionInteraction_LocalFrameMode LocalFrameMode = EModelingSelectionInteraction_LocalFrameMode::FromGeometry;
+
 	UPROPERTY()
 	TObjectPtr<URectangleMarqueeInteraction> RectangleMarqueeInteraction;
 
diff --git a/Engine/Plugins/Editor/ModelingToolsEditorMode/Source/ModelingToolsEditorMode/Public/ModelingToolsEditorModeSettings.h b/Engine/Plugins/Editor/ModelingToolsEditorMode/Source/ModelingToolsEditorMode/Public/ModelingToolsEditorModeSettings.h
--- a/Engine/Plugins/Editor/ModelingToolsEditorMode/Source/ModelingToolsEditorMode/Public/ModelingToolsEditorModeSettings.h
+++ b/Engine/Plugins/Editor/ModelingToolsEditorMode/Source/ModelingToolsEditorMode/Public/ModelingToolsEditorModeSettings.h
@@ -358,6 +358,9 @@ public:
 	UPROPERTY(config)
 	int32 LastMeshSelectionDragMode = 0;
 
+	UPROPERTY(config)
+	int32 LastMeshSelectionLocalFrameMode = 0;
+
 public:
 
 	// saved-state for various mode settings that does not persist between editor runs
diff --git a/Engine/Plugins/Runtime/MeshModelingToolset/Source/ModelingComponents/Private/Selection/DynamicMeshSelector.cpp b/Engine/Plugins/Runtime/MeshModelingToolset/Source/ModelingComponents/Private/Selection/DynamicMeshSelector.cpp
--- a/Engine/Plugins/Runtime/MeshModelingToolset/Source/ModelingComponents/Private/Selection/DynamicMeshSelector.cpp
+++ b/Engine/Plugins/Runtime/MeshModelingToolset/Source/ModelingComponents/Private/Selection/DynamicMeshSelector.cpp
@@ -569,10 +569,15 @@ void FBaseDynamicMeshSelector::GetSelectionFrame(const FGeometrySelection& Selec
 
 	if (bTransformToWorld)
 	{
-		SelectionFrame.Transform(GetLocalToWorldTransform());
+		GetTargetFrame(SelectionFrame);
 	}
 }
 
+void FBaseDynamicMeshSelector::GetTargetFrame(FFrame3d& SelectionFrame)
+{
+	SelectionFrame.Transform(GetLocalToWorldTransform());
+}
+
 void FBaseDynamicMeshSelector::AccumulateSelectionBounds(const FGeometrySelection& Selection, FGeometrySelectionBounds& BoundsInOut, bool bTransformToWorld)
 {
 	FTransform UseTransform = (bTransformToWorld) ? GetLocalToWorldTransform() : FTransform::Identity;
diff --git a/Engine/Plugins/Runtime/MeshModelingToolset/Source/ModelingComponents/Private/Selection/GeometrySelectionManager.cpp b/Engine/Plugins/Runtime/MeshModelingToolset/Source/ModelingComponents/Private/Selection/GeometrySelectionManager.cpp
--- a/Engine/Plugins/Runtime/MeshModelingToolset/Source/ModelingComponents/Private/Selection/GeometrySelectionManager.cpp
+++ b/Engine/Plugins/Runtime/MeshModelingToolset/Source/ModelingComponents/Private/Selection/GeometrySelectionManager.cpp
@@ -1019,6 +1019,17 @@ void UGeometrySelectionManager::GetSelectionWorldFrame(FFrame3d& SelectionFrame)
 	}
 }
 
+void UGeometrySelectionManager::GetTargetWorldFrame(FFrame3d& SelectionFrame) const
+{
+	SelectionFrame = FFrame3d();
+	if (HasSelection())
+	{
+		// only handling one target for now
+		//if (ActiveTargetReferences.Num() == 1)
+		TSharedPtr<FGeometrySelectionTarget> Target = ActiveTargetReferences[0];
+		Target->Selector->GetTargetFrame(SelectionFrame);
+	}
+}
 
 bool UGeometrySelectionManager::HasSelectionForComponent(UPrimitiveComponent* Component) const
 {
diff --git a/Engine/Plugins/Runtime/MeshModelingToolset/Source/ModelingComponents/Public/Selection/DynamicMeshSelector.h b/Engine/Plugins/Runtime/MeshModelingToolset/Source/ModelingComponents/Public/Selection/DynamicMeshSelector.h
--- a/Engine/Plugins/Runtime/MeshModelingToolset/Source/ModelingComponents/Public/Selection/DynamicMeshSelector.h
+++ b/Engine/Plugins/Runtime/MeshModelingToolset/Source/ModelingComponents/Public/Selection/DynamicMeshSelector.h
@@ -116,6 +116,7 @@ public:
 	}
 
 	virtual void GetSelectionFrame(const FGeometrySelection& Selection, UE::Geometry::FFrame3d& SelectionFrame, bool bTransformToWorld) override;
+	virtual void GetTargetFrame(UE::Geometry::FFrame3d& SelectionFrame) override;
 	virtual void AccumulateSelectionBounds(const FGeometrySelection& Selection, FGeometrySelectionBounds& BoundsInOut, bool bTransformToWorld) override;
 	virtual void AccumulateSelectionElements(const FGeometrySelection& Selection, FGeometrySelectionElements& Elements, bool bTransformToWorld, bool bIsForPreview) override;
 	
diff --git a/Engine/Plugins/Runtime/MeshModelingToolset/Source/ModelingComponents/Public/Selection/GeometrySelectionManager.h b/Engine/Plugins/Runtime/MeshModelingToolset/Source/ModelingComponents/Public/Selection/GeometrySelectionManager.h
--- a/Engine/Plugins/Runtime/MeshModelingToolset/Source/ModelingComponents/Public/Selection/GeometrySelectionManager.h
+++ b/Engine/Plugins/Runtime/MeshModelingToolset/Source/ModelingComponents/Public/Selection/GeometrySelectionManager.h
@@ -302,9 +302,13 @@ public:
 
 	/** @return a world-space bounding box for the active element selection */
 	virtual bool GetSelectionBounds(FGeometrySelectionBounds& BoundsOut) const;
+
 	/** @return a 3D transformation frame suitable for use with the active element selection */
 	virtual void GetSelectionWorldFrame(UE::Geometry::FFrame3d& SelectionFrame) const;
 
+	/** @return a 3D transformation frame suitable for use with the set of active targets */
+	virtual void GetTargetWorldFrame(UE::Geometry::FFrame3d& SelectionFrame) const;
+
 	/** @return true if there is an active IGeometrySelector target for the given Component and it has a non-empty selection */
 	virtual bool HasSelectionForComponent(UPrimitiveComponent* Component) const;
 	/** 
diff --git a/Engine/Plugins/Runtime/MeshModelingToolset/Source/ModelingComponents/Public/Selection/GeometrySelector.h b/Engine/Plugins/Runtime/MeshModelingToolset/Source/ModelingComponents/Public/Selection/GeometrySelector.h
--- a/Engine/Plugins/Runtime/MeshModelingToolset/Source/ModelingComponents/Public/Selection/GeometrySelector.h
+++ b/Engine/Plugins/Runtime/MeshModelingToolset/Source/ModelingComponents/Public/Selection/GeometrySelector.h
@@ -397,6 +397,11 @@ public:
 	 */
 	virtual void GetSelectionFrame(const FGeometrySelection& Selection, UE::Geometry::FFrame3d& SelectionFrameOut, bool bTransformToWorld) = 0;
 
+	/**
+	 * Compute a 3D frame for the Selector's target object
+	 */
+	virtual void GetTargetFrame(UE::Geometry::FFrame3d& SelectionFrameOut) = 0;
+
 	/**
 	 * Accumulate the bounds of the provided Selection in the provided BoundsInOut. BoundsInOut is not cleared.
 	 * @param bTransformToWorld if true each bounded point/element will be transformed to World space before being "contained" in BoundsInOut
