ModelingTools: Adding icons for various modeling selection activities and PolyEd activities.

diff --git a/Engine/Plugins/Editor/ModelingToolsEditorMode/Source/ModelingToolsEditorMode/Private/ModelingToolsEditorModeStyle.cpp b/Engine/Plugins/Editor/ModelingToolsEditorMode/Source/ModelingToolsEditorMode/Private/ModelingToolsEditorModeStyle.cpp
--- a/Engine/Plugins/Editor/ModelingToolsEditorMode/Source/ModelingToolsEditorMode/Private/ModelingToolsEditorModeStyle.cpp
+++ b/Engine/Plugins/Editor/ModelingToolsEditorMode/Source/ModelingToolsEditorMode/Private/ModelingToolsEditorModeStyle.cpp
@@ -320,6 +320,8 @@ void FModelingToolsEditorModeStyle::Initialize()
 		StyleSet->Set("ModelingToolsManagerCommands.BeginSplitMeshesTool.Small", new IMAGE_BRUSH_SVG("Icons/GeometrySplit", Icon20x20));
 		StyleSet->Set("ModelingToolsManagerCommands.BeginPatternTool", new IMAGE_BRUSH_SVG("Icons/ModelingPattern", Icon20x20));
 		StyleSet->Set("ModelingToolsManagerCommands.BeginPatternTool.Small", new IMAGE_BRUSH_SVG("Icons/ModelingPattern", Icon20x20));
+		StyleSet->Set("ModelingToolsManagerCommands.BeginHarvestInstancesTool", new IMAGE_BRUSH_SVG("Icons/HarvestInstances", Icon20x20));
+		StyleSet->Set("ModelingToolsManagerCommands.BeginHarvestInstancesTool.Small", new IMAGE_BRUSH_SVG("Icons/HarvestInstances", Icon20x20));
 		StyleSet->Set("ModelingToolsManagerCommands.BeginMeshVertexPaintTool", new IMAGE_BRUSH_SVG("Icons/PaintVertexColors", Icon20x20));
 		StyleSet->Set("ModelingToolsManagerCommands.BeginMeshVertexPaintTool.Small", new IMAGE_BRUSH_SVG("Icons/PaintVertexColors", Icon20x20));
 
@@ -372,11 +374,31 @@ void FModelingToolsEditorModeStyle::Initialize()
 		StyleSet->Set("ModelingToolsManagerCommands.BeginSelectionAction_Contract", new IMAGE_BRUSH_SVG("Icons/ModSelectionShrink_16", Icon20x20));
 
 		StyleSet->Set("ModelingToolsManagerCommands.BeginSelectionAction_Delete", new FSlateImageBrush(StyleSet->RootToCoreContentDir(TEXT("../Editor/Slate/Icons/GeneralTools/Delete_40x.png")), Icon20x20));
-
+		StyleSet->Set("ModelingToolsManagerCommands.BeginSelectionAction_Delete.Small", new FSlateImageBrush(StyleSet->RootToCoreContentDir(TEXT("../Editor/Slate/Icons/GeneralTools/Delete_40x.png")), Icon20x20));
+		StyleSet->Set("ModelingToolsManagerCommands.BeginSelectionAction_Extrude", new IMAGE_BRUSH_SVG("Icons/Extrude", Icon20x20));
+		StyleSet->Set("ModelingToolsManagerCommands.BeginSelectionAction_Extrude.Small", new IMAGE_BRUSH_SVG("Icons/Extrude", Icon20x20));
+		StyleSet->Set("ModelingToolsManagerCommands.BeginSelectionAction_Offset", new IMAGE_BRUSH_SVG("Icons/Offset", Icon20x20));
+		StyleSet->Set("ModelingToolsManagerCommands.BeginSelectionAction_Offset.Small", new IMAGE_BRUSH_SVG("Icons/Offset", Icon20x20));
 		StyleSet->Set("ModelingToolsManagerCommands.BeginPolyModelTool_PolyEd", 				new IMAGE_PLUGIN_BRUSH("Icons/PolyEdit_40x", 		Icon20x20));
 		StyleSet->Set("ModelingToolsManagerCommands.BeginPolyModelTool_PolyEd.Small", 			new IMAGE_PLUGIN_BRUSH("Icons/PolyEdit_40x", 		Icon20x20));
 		StyleSet->Set("ModelingToolsManagerCommands.BeginPolyModelTool_TriSel", 				new IMAGE_PLUGIN_BRUSH("Icons/MeshSelect_40x",		Icon20x20));
 		StyleSet->Set("ModelingToolsManagerCommands.BeginPolyModelTool_TriSel.Small", 			new IMAGE_PLUGIN_BRUSH("Icons/MeshSelect_40x",		Icon20x20));
+		StyleSet->Set("ModelingToolsManagerCommands.BeginPolyModelTool_Inset", new IMAGE_BRUSH_SVG("Icons/Inset", Icon20x20));
+		StyleSet->Set("ModelingToolsManagerCommands.BeginPolyModelTool_Inset.Small", new IMAGE_BRUSH_SVG("Icons/Inset", Icon20x20));
+		StyleSet->Set("ModelingToolsManagerCommands.BeginPolyModelTool_Outset", new IMAGE_BRUSH_SVG("Icons/Outset", Icon20x20));
+		StyleSet->Set("ModelingToolsManagerCommands.BeginPolyModelTool_Outset.Small", new IMAGE_BRUSH_SVG("Icons/Outset", Icon20x20));
+		StyleSet->Set("ModelingToolsManagerCommands.BeginPolyModelTool_CutFaces", new IMAGE_BRUSH_SVG("Icons/CutFaces", Icon20x20));
+		StyleSet->Set("ModelingToolsManagerCommands.BeginPolyModelTool_CutFaces.Small", new IMAGE_BRUSH_SVG("Icons/CutFaces", Icon20x20));
+		StyleSet->Set("ModelingToolsManagerCommands.BeginPolyModelTool_InsertEdgeLoop", new IMAGE_PLUGIN_BRUSH("Icons/ModelingEdgeLoopInsert_x40", Icon20x20));
+		StyleSet->Set("ModelingToolsManagerCommands.BeginPolyModelTool_InsertEdgeLoop.Small", new IMAGE_PLUGIN_BRUSH("Icons/ModelingEdgeLoopInsert_x40", Icon20x20));
+		StyleSet->Set("ModelingToolsManagerCommands.BeginPolyModelTool_ExtrudeEdges", new IMAGE_BRUSH_SVG("Icons/ExtrudeEdge", Icon20x20));
+		StyleSet->Set("ModelingToolsManagerCommands.BeginPolyModelTool_ExtrudeEdges.Small", new IMAGE_BRUSH_SVG("Icons/ExtrudeEdge", Icon20x20));
+		StyleSet->Set("ModelingToolsManagerCommands.BeginPolyModelTool_PushPull", new IMAGE_BRUSH_SVG("Icons/PushPull", Icon20x20));
+		StyleSet->Set("ModelingToolsManagerCommands.BeginPolyModelTool_PushPull.Small", new IMAGE_BRUSH_SVG("Icons/PushPull", Icon20x20));
+		StyleSet->Set("ModelingToolsManagerCommands.BeginPolyModelTool_Bevel", new IMAGE_BRUSH_SVG("Icons/Bevel", Icon20x20));
+		StyleSet->Set("ModelingToolsManagerCommands.BeginPolyModelTool_Bevel.Small", new IMAGE_BRUSH_SVG("Icons/Bevel", Icon20x20));
+		StyleSet->Set("ModelingToolsManagerCommands.BeginSelectionAction_Retriangulate", new IMAGE_BRUSH_SVG("Icons/Clean", Icon20x20));
+		StyleSet->Set("ModelingToolsManagerCommands.BeginSelectionAction_Retriangulate.Small", new IMAGE_BRUSH_SVG("Icons/Clean", Icon20x20));
 
 		StyleSet->Set("ModelingModeSelection.More_Right",  new IMAGE_BRUSH_SVG("Icons/SelectionToolbar_More", Icon20x20));
 		StyleSet->Set("ModelingModeSelection.Edits_Right",  new IMAGE_BRUSH_SVG("Icons/SelectionToolbar_Edits", Icon20x20));
@@ -463,6 +485,15 @@ void FModelingToolsEditorModeStyle::Initialize()
 		// Icons for PolyEd and TriEd activities
 		StyleSet->Set("PolyEd.InsertGroupEdge", new IMAGE_PLUGIN_BRUSH("Icons/ModelingGroupEdgeInsert_x40", Icon20x20));
 		StyleSet->Set("PolyEd.InsertEdgeLoop", new IMAGE_PLUGIN_BRUSH("Icons/ModelingEdgeLoopInsert_x40", Icon20x20));
+		StyleSet->Set("PolyEd.Extrude", new IMAGE_BRUSH_SVG("Icons/Extrude", Icon20x20));
+		StyleSet->Set("PolyEd.Offset", new IMAGE_BRUSH_SVG("Icons/Offset", Icon20x20));
+		StyleSet->Set("PolyEd.PushPull", new IMAGE_BRUSH_SVG("Icons/PushPull", Icon20x20));
+		StyleSet->Set("PolyEd.Inset", new IMAGE_BRUSH_SVG("Icons/Inset", Icon20x20));
+		StyleSet->Set("PolyEd.Outset", new IMAGE_BRUSH_SVG("Icons/Outset", Icon20x20));
+		StyleSet->Set("PolyEd.CutFaces", new IMAGE_BRUSH_SVG("Icons/CutFaces", Icon20x20));
+		StyleSet->Set("PolyEd.ProjectUVs", new IMAGE_PLUGIN_BRUSH("Icons/UVProjection_40x", Icon20x20));
+		StyleSet->Set("PolyEd.Bevel", new IMAGE_BRUSH_SVG("Icons/Bevel", Icon20x20));
+		StyleSet->Set("PolyEd.ExtrudeEdge", new IMAGE_BRUSH_SVG("Icons/ExtrudeEdge", Icon20x20));
 	}
 
 	// Style for the toolbar in the PolyEd customization
diff --git a/Engine/Plugins/Runtime/MeshModelingToolset/Source/MeshModelingTools/Private/EditMeshPolygonsTool.cpp b/Engine/Plugins/Runtime/MeshModelingToolset/Source/MeshModelingTools/Private/EditMeshPolygonsTool.cpp
--- a/Engine/Plugins/Runtime/MeshModelingToolset/Source/MeshModelingTools/Private/EditMeshPolygonsTool.cpp
+++ b/Engine/Plugins/Runtime/MeshModelingToolset/Source/MeshModelingTools/Private/EditMeshPolygonsTool.cpp
@@ -603,8 +603,8 @@ void UEditMeshPolygonsTool::Setup()
 
 	ExtrudeActivity = NewObject<UPolyEditExtrudeActivity>();
 	ExtrudeActivity->Setup(this);
-	ActivityLabels.Add(ExtrudeActivity, LOCTEXT("ExtrudeActivityLabel", "Extrude"));
-	ActivityIconNames.Add(ExtrudeActivity, "PolyEd.Extrude");
+	// The icons/labels differ depending on whether we're doing extrude, offset, or push/pull, so
+	// set those when we launch the activity.
 	
 	InsetOutsetActivity = NewObject<UPolyEditInsetOutsetActivity>();
 	InsetOutsetActivity->Setup(this);
@@ -1251,6 +1251,10 @@ void UEditMeshPolygonsTool::OnTick(float DeltaTime)
 			}
 			ExtrudeActivity->ExtrudeMode = FExtrudeOp::EExtrudeMode::MoveAndStitch;
 			ExtrudeActivity->PropertySetToUse = UPolyEditExtrudeActivity::EPropertySetToUse::Extrude;
+
+			ActivityLabels.Add(ExtrudeActivity, LOCTEXT("ExtrudeActivityLabel", "Extrude"));
+			ActivityIconNames.Add(ExtrudeActivity, "PolyEd.Extrude");
+
 			StartActivity(ExtrudeActivity);
 			break;
 		}
@@ -1258,6 +1262,10 @@ void UEditMeshPolygonsTool::OnTick(float DeltaTime)
 		{
 			ExtrudeActivity->ExtrudeMode = FExtrudeOp::EExtrudeMode::Boolean;
 			ExtrudeActivity->PropertySetToUse = UPolyEditExtrudeActivity::EPropertySetToUse::PushPull;
+
+			ActivityLabels.Add(ExtrudeActivity, LOCTEXT("PushPullActivityLabel", "Push/Pull"));
+			ActivityIconNames.Add(ExtrudeActivity, "PolyEd.PushPull");
+
 			StartActivity(ExtrudeActivity);
 			break;
 		}
@@ -1265,6 +1273,10 @@ void UEditMeshPolygonsTool::OnTick(float DeltaTime)
 		{
 			ExtrudeActivity->ExtrudeMode = FExtrudeOp::EExtrudeMode::MoveAndStitch;
 			ExtrudeActivity->PropertySetToUse = UPolyEditExtrudeActivity::EPropertySetToUse::Offset;
+
+			ActivityLabels.Add(ExtrudeActivity, LOCTEXT("OffsetActivityLabel", "Offset"));
+			ActivityIconNames.Add(ExtrudeActivity, "PolyEd.Offset");
+
 			StartActivity(ExtrudeActivity);
 			break;
 		}
