ModelingMode: Update missing and incorrect icons for mesh selection selection mode.

diff --git a/Engine/Plugins/Editor/ModelingToolsEditorMode/Source/ModelingToolsEditorMode/Private/ModelingToolsEditorModeStyle.cpp b/Engine/Plugins/Editor/ModelingToolsEditorMode/Source/ModelingToolsEditorMode/Private/ModelingToolsEditorModeStyle.cpp
--- a/Engine/Plugins/Editor/ModelingToolsEditorMode/Source/ModelingToolsEditorMode/Private/ModelingToolsEditorModeStyle.cpp
+++ b/Engine/Plugins/Editor/ModelingToolsEditorMode/Source/ModelingToolsEditorMode/Private/ModelingToolsEditorModeStyle.cpp
@@ -373,24 +373,24 @@ void FModelingToolsEditorModeStyle::Initialize()
 		StyleSet->Set("ModelingToolsManagerCommands.BeginSelectionAction_Expand", new IMAGE_BRUSH_SVG("Icons/ModSelectionExpand_16", Icon20x20));
 		StyleSet->Set("ModelingToolsManagerCommands.BeginSelectionAction_Contract", new IMAGE_BRUSH_SVG("Icons/ModSelectionShrink_16", Icon20x20));
 
-		StyleSet->Set("ModelingToolsManagerCommands.BeginSelectionAction_Delete", new FSlateImageBrush(StyleSet->RootToCoreContentDir(TEXT("../Editor/Slate/Icons/GeneralTools/Delete_40x.png")), Icon20x20));
-		StyleSet->Set("ModelingToolsManagerCommands.BeginSelectionAction_Delete.Small", new FSlateImageBrush(StyleSet->RootToCoreContentDir(TEXT("../Editor/Slate/Icons/GeneralTools/Delete_40x.png")), Icon20x20));
+		StyleSet->Set("ModelingToolsManagerCommands.BeginSelectionAction_Delete", new IMAGE_BRUSH_SVG("Icons/Delete", Icon20x20));
+		StyleSet->Set("ModelingToolsManagerCommands.BeginSelectionAction_Delete.Small", new IMAGE_BRUSH_SVG("Icons/Delete", Icon20x20));
 		StyleSet->Set("ModelingToolsManagerCommands.BeginSelectionAction_Extrude", new IMAGE_BRUSH_SVG("Icons/Extrude", Icon20x20));
 		StyleSet->Set("ModelingToolsManagerCommands.BeginSelectionAction_Extrude.Small", new IMAGE_BRUSH_SVG("Icons/Extrude", Icon20x20));
 		StyleSet->Set("ModelingToolsManagerCommands.BeginSelectionAction_Offset", new IMAGE_BRUSH_SVG("Icons/Offset", Icon20x20));
 		StyleSet->Set("ModelingToolsManagerCommands.BeginSelectionAction_Offset.Small", new IMAGE_BRUSH_SVG("Icons/Offset", Icon20x20));
-		StyleSet->Set("ModelingToolsManagerCommands.BeginPolyModelTool_PolyEd", 				new IMAGE_PLUGIN_BRUSH("Icons/PolyEdit_40x", 		Icon20x20));
-		StyleSet->Set("ModelingToolsManagerCommands.BeginPolyModelTool_PolyEd.Small", 			new IMAGE_PLUGIN_BRUSH("Icons/PolyEdit_40x", 		Icon20x20));
-		StyleSet->Set("ModelingToolsManagerCommands.BeginPolyModelTool_TriSel", 				new IMAGE_PLUGIN_BRUSH("Icons/MeshSelect_40x",		Icon20x20));
-		StyleSet->Set("ModelingToolsManagerCommands.BeginPolyModelTool_TriSel.Small", 			new IMAGE_PLUGIN_BRUSH("Icons/MeshSelect_40x",		Icon20x20));
+		StyleSet->Set("ModelingToolsManagerCommands.BeginPolyModelTool_PolyEd", new IMAGE_BRUSH_SVG("Icons/PolyEdit", Icon20x20));
+		StyleSet->Set("ModelingToolsManagerCommands.BeginPolyModelTool_PolyEd.Small", new IMAGE_BRUSH_SVG("Icons/PolyEdit",	Icon20x20));
+		StyleSet->Set("ModelingToolsManagerCommands.BeginPolyModelTool_TriSel", new IMAGE_BRUSH_SVG("Icons/TriSelect", Icon20x20));
+		StyleSet->Set("ModelingToolsManagerCommands.BeginPolyModelTool_TriSel.Small", new IMAGE_BRUSH_SVG("Icons/TriSelect", Icon20x20));
 		StyleSet->Set("ModelingToolsManagerCommands.BeginPolyModelTool_Inset", new IMAGE_BRUSH_SVG("Icons/Inset", Icon20x20));
 		StyleSet->Set("ModelingToolsManagerCommands.BeginPolyModelTool_Inset.Small", new IMAGE_BRUSH_SVG("Icons/Inset", Icon20x20));
 		StyleSet->Set("ModelingToolsManagerCommands.BeginPolyModelTool_Outset", new IMAGE_BRUSH_SVG("Icons/Outset", Icon20x20));
 		StyleSet->Set("ModelingToolsManagerCommands.BeginPolyModelTool_Outset.Small", new IMAGE_BRUSH_SVG("Icons/Outset", Icon20x20));
 		StyleSet->Set("ModelingToolsManagerCommands.BeginPolyModelTool_CutFaces", new IMAGE_BRUSH_SVG("Icons/CutFaces", Icon20x20));
 		StyleSet->Set("ModelingToolsManagerCommands.BeginPolyModelTool_CutFaces.Small", new IMAGE_BRUSH_SVG("Icons/CutFaces", Icon20x20));
-		StyleSet->Set("ModelingToolsManagerCommands.BeginPolyModelTool_InsertEdgeLoop", new IMAGE_PLUGIN_BRUSH("Icons/ModelingEdgeLoopInsert_x40", Icon20x20));
-		StyleSet->Set("ModelingToolsManagerCommands.BeginPolyModelTool_InsertEdgeLoop.Small", new IMAGE_PLUGIN_BRUSH("Icons/ModelingEdgeLoopInsert_x40", Icon20x20));
+		StyleSet->Set("ModelingToolsManagerCommands.BeginPolyModelTool_InsertEdgeLoop", new IMAGE_BRUSH_SVG("Icons/EdgeLoop", Icon20x20));
+		StyleSet->Set("ModelingToolsManagerCommands.BeginPolyModelTool_InsertEdgeLoop.Small", new IMAGE_BRUSH_SVG("Icons/EdgeLoop", Icon20x20));
 		StyleSet->Set("ModelingToolsManagerCommands.BeginPolyModelTool_ExtrudeEdges", new IMAGE_BRUSH_SVG("Icons/ExtrudeEdge", Icon20x20));
 		StyleSet->Set("ModelingToolsManagerCommands.BeginPolyModelTool_ExtrudeEdges.Small", new IMAGE_BRUSH_SVG("Icons/ExtrudeEdge", Icon20x20));
 		StyleSet->Set("ModelingToolsManagerCommands.BeginPolyModelTool_PushPull", new IMAGE_BRUSH_SVG("Icons/PushPull", Icon20x20));
