Add the primitive component tool target to modeling mode, so that e.g. the transform tool works on more targets

diff --git a/Engine/Plugins/Editor/ModelingToolsEditorMode/Source/ModelingToolsEditorMode/Private/ModelingToolsEditorMode.cpp b/Engine/Plugins/Editor/ModelingToolsEditorMode/Source/ModelingToolsEditorMode/Private/ModelingToolsEditorMode.cpp
--- a/Engine/Plugins/Editor/ModelingToolsEditorMode/Source/ModelingToolsEditorMode/Private/ModelingToolsEditorMode.cpp
+++ b/Engine/Plugins/Editor/ModelingToolsEditorMode/Source/ModelingToolsEditorMode/Private/ModelingToolsEditorMode.cpp
@@ -391,6 +391,9 @@ void UModelingToolsEditorMode::Enter()
 	// skin weights.
 	GetInteractiveToolsContext()->TargetManager->AddTargetFactory(NewObject<USkeletalMeshComponentReadOnlyToolTargetFactory>(GetToolManager()));
 
+	// Register builders for the generic component tool target, to support tools that only need the primitive component (e.g. the transform tool)
+	GetInteractiveToolsContext()->TargetManager->AddTargetFactory(NewObject<UPrimitiveComponentToolTargetFactory>(GetToolManager()));
+
 	// listen to post-build
 	GetToolManager()->OnToolPostBuild.AddUObject(this, &UModelingToolsEditorMode::OnToolPostBuild);
 
