update split tool description text to reflect new features

diff --git a/Engine/Plugins/Editor/ModelingToolsEditorMode/Source/ModelingToolsEditorMode/Private/ModelingToolsManagerActions.cpp b/Engine/Plugins/Editor/ModelingToolsEditorMode/Source/ModelingToolsEditorMode/Private/ModelingToolsManagerActions.cpp
--- a/Engine/Plugins/Editor/ModelingToolsEditorMode/Source/ModelingToolsEditorMode/Private/ModelingToolsManagerActions.cpp
+++ b/Engine/Plugins/Editor/ModelingToolsEditorMode/Source/ModelingToolsEditorMode/Private/ModelingToolsManagerActions.cpp
@@ -175,7 +175,7 @@ void FModelingToolsManagerCommands::RegisterCommands()
 	REGISTER_MODELING_TOOL_COMMAND(BeginAlignObjectsTool, "Align", "Align the selected Objects");
 	REGISTER_MODELING_TOOL_COMMAND(BeginTransferMeshTool, "Transfer", "Copy the first selected Mesh to the second selected Mesh");
 	REGISTER_MODELING_TOOL_COMMAND(BeginConvertMeshesTool, "Convert", "Convert the selected Meshes to a different type of Mesh Object");
-	REGISTER_MODELING_TOOL_COMMAND(BeginSplitMeshesTool, "Split", "Split the selected Meshes into separate parts based on connectivity");
+	REGISTER_MODELING_TOOL_COMMAND(BeginSplitMeshesTool, "Split", "Split the selected Meshes into separate parts based on connectivity, selection, material ID or PolyGroup");
 	REGISTER_MODELING_TOOL_COMMAND(BeginPatternTool, "Pattern", "Create patterns of Meshes");
 	REGISTER_MODELING_TOOL_COMMAND(BeginHarvestInstancesTool, "Harvest Instances", "Extract a set of InstancedStaticMeshComponents from a set of Actors");
 	
