Modeling Tools: Fixes a bug that ignored MeshElementSelection colors when exiting and re-entering modeling mode

diff --git a/Engine/Plugins/Editor/ModelingToolsEditorMode/Source/ModelingToolsEditorMode/Private/ModelingToolsEditorMode.cpp b/Engine/Plugins/Editor/ModelingToolsEditorMode/Source/ModelingToolsEditorMode/Private/ModelingToolsEditorMode.cpp
--- a/Engine/Plugins/Editor/ModelingToolsEditorMode/Source/ModelingToolsEditorMode/Private/ModelingToolsEditorMode.cpp
+++ b/Engine/Plugins/Editor/ModelingToolsEditorMode/Source/ModelingToolsEditorMode/Private/ModelingToolsEditorMode.cpp
@@ -882,7 +882,6 @@ void UModelingToolsEditorMode::Enter()
 						ModelingEditorSettings->LastMeshSelectionTopologyMode = static_cast<int>(TopoMode);
 						ModelingEditorSettings->LastMeshSelectionElementType = static_cast<int>(ElementMode);
 						ModelingEditorSettings->SaveConfig();
-						SelectionManager->SetSelectionColors(ModelingEditorSettings->UnselectedColor, ModelingEditorSettings->HoverOverSelectedColor, ModelingEditorSettings->HoverOverUnselectedColor, ModelingEditorSettings->GeometrySelectedColor);
 					}
 				}
 			}),
@@ -901,6 +900,8 @@ void UModelingToolsEditorMode::Enter()
 		RegisterSelectionMode(UGeometrySelectionManager::EMeshTopologyMode::Polygroup, UGeometrySelectionManager::EGeometryElementType::Edge, ToolManagerCommands.MeshSelectionModeAction_GroupEdges);
 	}
 
+	// Colors initialized here any time Modeling mode is entered
+	SelectionManager->SetSelectionColors(ModelingEditorSettings->UnselectedColor, ModelingEditorSettings->HoverOverSelectedColor, ModelingEditorSettings->HoverOverUnselectedColor, ModelingEditorSettings->GeometrySelectedColor);
 
 	// this function registers and tracks an active UGeometrySelectionEditCommand and it's associated UICommand
 	auto RegisterSelectionCommand = [&](UGeometrySelectionEditCommand* Command, TSharedPtr<FUICommandInfo> UICommand, bool bAlwaysVisible)
diff --git a/Engine/Plugins/Editor/ModelingToolsEditorMode/Source/ModelingToolsEditorMode/Private/ModelingToolsEditorModeToolkit.cpp b/Engine/Plugins/Editor/ModelingToolsEditorMode/Source/ModelingToolsEditorMode/Private/ModelingToolsEditorModeToolkit.cpp
--- a/Engine/Plugins/Editor/ModelingToolsEditorMode/Source/ModelingToolsEditorMode/Private/ModelingToolsEditorModeToolkit.cpp
+++ b/Engine/Plugins/Editor/ModelingToolsEditorMode/Source/ModelingToolsEditorMode/Private/ModelingToolsEditorModeToolkit.cpp
@@ -64,6 +64,7 @@
 #include "ModelingWidgets/SToolInputAssetComboPanel.h"
 #include "Fonts/SlateFontInfo.h"
 #include "ToolPresetAssetSubsystem.h"
+#include "Selection/GeometrySelectionManager.h"
 #include "Widgets/Input/SMultiLineEditableTextBox.h"
 
 
@@ -642,6 +643,9 @@ void FModelingToolsEditorModeToolkit::Init(const TSharedPtr<IToolkitHost>& InitT
 	GetToolkitHost()->AddViewportOverlayWidget(SelectionPaletteOverlayWidget.ToSharedRef());
 
 	CurrentPresetPath = FSoftObjectPath(); // Default to the default collection by leaving this null.
+
+	UModelingToolsModeCustomizationSettings* ModelingCustomizationSettings = GetMutableDefault<UModelingToolsModeCustomizationSettings>();
+	ModelingCustomizationSettings->OnSettingChanged().AddSP(SharedThis(this),  &FModelingToolsEditorModeToolkit::UpdateSelectionColors);
 }
 
 void FModelingToolsEditorModeToolkit::MakeToolShutdownOverlayWidget()
@@ -2535,6 +2539,11 @@ void FModelingToolsEditorModeToolkit::UpdateCategoryButtonLabelVisibility(UObjec
 	ToolkitBuilder->RefreshCategoryToolbarWidget();
 }
 
+void FModelingToolsEditorModeToolkit::UpdateSelectionColors(UObject* Obj, FPropertyChangedEvent& ChangeEvent) const
+{
+	const UModelingToolsModeCustomizationSettings* UISettings = GetMutableDefault<UModelingToolsModeCustomizationSettings>();
+	CachedSelectionManager->SetSelectionColors(UISettings->UnselectedColor, UISettings->HoverOverSelectedColor, UISettings->HoverOverUnselectedColor, UISettings->GeometrySelectedColor);
+}
 
 FText FModelingToolsEditorModeToolkit::GetRestrictiveModeAutoGeneratedAssetPathText() const
 {
diff --git a/Engine/Plugins/Editor/ModelingToolsEditorMode/Source/ModelingToolsEditorMode/Public/ModelingToolsEditorModeToolkit.h b/Engine/Plugins/Editor/ModelingToolsEditorMode/Source/ModelingToolsEditorMode/Public/ModelingToolsEditorModeToolkit.h
--- a/Engine/Plugins/Editor/ModelingToolsEditorMode/Source/ModelingToolsEditorMode/Public/ModelingToolsEditorModeToolkit.h
+++ b/Engine/Plugins/Editor/ModelingToolsEditorMode/Source/ModelingToolsEditorMode/Public/ModelingToolsEditorModeToolkit.h
@@ -203,6 +203,11 @@ private:
 	 */
 	void UpdateCategoryButtonLabelVisibility(UObject* Obj, FPropertyChangedEvent& ChangeEvent);
 
+	/**
+	 *  Updates the mesh element selection colors based off of the current editor preference for it
+	 */
+	void UpdateSelectionColors(UObject* Obj, FPropertyChangedEvent& ChangeEvent) const;
+
 	TArray<TSharedPtr<FString>> AssetLODModes;
 	TSharedPtr<STextBlock> AssetLODModeLabel;
 	TSharedPtr<STextComboBox> AssetLODMode;
