Modeling Tools: Fixes bug with selection color not clearing in Mesh Element Selection

diff --git a/Engine/Plugins/Editor/ModelingToolsEditorMode/Source/ModelingToolsEditorMode/Private/ModelingToolsEditorMode.cpp b/Engine/Plugins/Editor/ModelingToolsEditorMode/Source/ModelingToolsEditorMode/Private/ModelingToolsEditorMode.cpp
--- a/Engine/Plugins/Editor/ModelingToolsEditorMode/Source/ModelingToolsEditorMode/Private/ModelingToolsEditorMode.cpp
+++ b/Engine/Plugins/Editor/ModelingToolsEditorMode/Source/ModelingToolsEditorMode/Private/ModelingToolsEditorMode.cpp
@@ -875,7 +875,7 @@ void UModelingToolsEditorMode::Enter()
 				if ( GetToolManager() && GetToolManager()->GetContextTransactionsAPI() && GetSelectionManager() )
 				{
 					GetToolManager()->GetContextTransactionsAPI()->BeginUndoTransaction(LOCTEXT("ChangeSelectionMode", "Selection Mode"));
-					GetSelectionManager()->SetMeshSelectionTypeAndMode(ElementMode, TopoMode, true);
+					GetSelectionManager()->SetMeshSelectionTypeAndMode(ElementMode, TopoMode,  TopoMode != UGeometrySelectionManager::EMeshTopologyMode::None);
 					GetToolManager()->GetContextTransactionsAPI()->EndUndoTransaction();
 					if (UModelingToolsModeCustomizationSettings* ModelingEditorSettings = GetMutableDefault<UModelingToolsModeCustomizationSettings>())
 					{
