Enabling the Stylus plugin

diff --git a/Engine/Plugins/Editor/ModelingToolsEditorMode/ModelingToolsEditorMode.uplugin b/Engine/Plugins/Editor/ModelingToolsEditorMode/ModelingToolsEditorMode.uplugin
--- a/Engine/Plugins/Editor/ModelingToolsEditorMode/ModelingToolsEditorMode.uplugin
+++ b/Engine/Plugins/Editor/ModelingToolsEditorMode/ModelingToolsEditorMode.uplugin
@@ -40,6 +40,11 @@
 			"Name": "ToolPresets",
 			"Enabled": true,
 			"TargetAllowList": [ "Editor" ]
+		},
+		{
+			"Name": "StylusInput",
+			"Enabled": true,
+			"TargetAllowList": [ "Editor" ]
 		}
 	]
 }
\ No newline at end of file
diff --git a/Engine/Plugins/Editor/ModelingToolsEditorMode/Source/ModelingToolsEditorMode/ModelingToolsEditorMode.Build.cs b/Engine/Plugins/Editor/ModelingToolsEditorMode/Source/ModelingToolsEditorMode/ModelingToolsEditorMode.Build.cs
--- a/Engine/Plugins/Editor/ModelingToolsEditorMode/Source/ModelingToolsEditorMode/ModelingToolsEditorMode.Build.cs
+++ b/Engine/Plugins/Editor/ModelingToolsEditorMode/Source/ModelingToolsEditorMode/ModelingToolsEditorMode.Build.cs
@@ -54,6 +54,7 @@ public class ModelingToolsEditorMode : ModuleRules
 				"ToolPresetAsset",
 				"ToolPresetEditor",
 				"EditorConfig",
+				"StylusInput",
 				// ... add private dependencies that you statically link with here ...	
 			}
 		);
diff --git a/Engine/Plugins/Editor/ModelingToolsEditorMode/Source/ModelingToolsEditorMode/Private/ModelingToolsEditorMode.cpp b/Engine/Plugins/Editor/ModelingToolsEditorMode/Source/ModelingToolsEditorMode/Private/ModelingToolsEditorMode.cpp
--- a/Engine/Plugins/Editor/ModelingToolsEditorMode/Source/ModelingToolsEditorMode/Private/ModelingToolsEditorMode.cpp
+++ b/Engine/Plugins/Editor/ModelingToolsEditorMode/Source/ModelingToolsEditorMode/Private/ModelingToolsEditorMode.cpp
@@ -120,9 +120,9 @@
 #include "EditorModeManager.h"
 #include "UnrealWidget.h"
 
-// Stylus support is currently disabled due to issues with the stylus plugin
-// We are leaving the code in this cpp file, defined out, so that it is easier to bring back if/when the stylus plugin is improved.
-#define ENABLE_STYLUS_SUPPORT 0
+// Enabling the Stylus implementation to test for possible issues
+// This implementation is likely to be updated
+#define ENABLE_STYLUS_SUPPORT 1
 
 #if ENABLE_STYLUS_SUPPORT 
 #include "IStylusState.h"
diff --git a/Engine/Plugins/Editor/ModelingToolsEditorMode/Source/ModelingToolsEditorMode/Public/ModelingToolsEditorMode.h b/Engine/Plugins/Editor/ModelingToolsEditorMode/Source/ModelingToolsEditorMode/Public/ModelingToolsEditorMode.h
--- a/Engine/Plugins/Editor/ModelingToolsEditorMode/Source/ModelingToolsEditorMode/Public/ModelingToolsEditorMode.h
+++ b/Engine/Plugins/Editor/ModelingToolsEditorMode/Source/ModelingToolsEditorMode/Public/ModelingToolsEditorMode.h
@@ -116,7 +116,7 @@ protected:
 	void OnEditorClosed();
 
 	// Stylus support is currently disabled; this is left in for reference if/when it is brought back
-	//TUniquePtr<FStylusStateTracker> StylusStateTracker;
+	TUniquePtr<FStylusStateTracker> StylusStateTracker;
 
 	TSharedPtr<FLevelObjectsObserver> LevelObjectsObserver;
 
