Modeling Tools: Fixes a bug which caused a crash when attempting to customize mesh element selection colors before enabling Mesh Element Selection toolbar

diff --git a/Engine/Plugins/Editor/ModelingToolsEditorMode/Source/ModelingToolsEditorMode/Private/ModelingToolsEditorModeToolkit.cpp b/Engine/Plugins/Editor/ModelingToolsEditorMode/Source/ModelingToolsEditorMode/Private/ModelingToolsEditorModeToolkit.cpp
--- a/Engine/Plugins/Editor/ModelingToolsEditorMode/Source/ModelingToolsEditorMode/Private/ModelingToolsEditorModeToolkit.cpp
+++ b/Engine/Plugins/Editor/ModelingToolsEditorMode/Source/ModelingToolsEditorMode/Private/ModelingToolsEditorModeToolkit.cpp
@@ -2541,8 +2541,11 @@ void FModelingToolsEditorModeToolkit::UpdateCategoryButtonLabelVisibility(UObjec
 
 void FModelingToolsEditorModeToolkit::UpdateSelectionColors(UObject* Obj, FPropertyChangedEvent& ChangeEvent) const
 {
-	const UModelingToolsModeCustomizationSettings* UISettings = GetMutableDefault<UModelingToolsModeCustomizationSettings>();
-	CachedSelectionManager->SetSelectionColors(UISettings->UnselectedColor, UISettings->HoverOverSelectedColor, UISettings->HoverOverUnselectedColor, UISettings->GeometrySelectedColor);
+	if (CachedSelectionManager)
+	{
+		const UModelingToolsModeCustomizationSettings* UISettings = GetMutableDefault<UModelingToolsModeCustomizationSettings>();
+		CachedSelectionManager->SetSelectionColors(UISettings->UnselectedColor, UISettings->HoverOverSelectedColor, UISettings->HoverOverUnselectedColor, UISettings->GeometrySelectedColor);
+	}
 }
 
 FText FModelingToolsEditorModeToolkit::GetRestrictiveModeAutoGeneratedAssetPathText() const
