make UEdMode handle interactive toolkit focus api's focus requests directly, to decouple it from editor selection + fix the bug that it doesn't work when there is no editor selection.
Also make modeling mode handle element selection through the same system.

diff --git a/Engine/Plugins/Editor/ModelingToolsEditorMode/Source/ModelingToolsEditorMode/Private/ModelingToolsEditorMode.cpp b/Engine/Plugins/Editor/ModelingToolsEditorMode/Source/ModelingToolsEditorMode/Private/ModelingToolsEditorMode.cpp
--- a/Engine/Plugins/Editor/ModelingToolsEditorMode/Source/ModelingToolsEditorMode/Private/ModelingToolsEditorMode.cpp
+++ b/Engine/Plugins/Editor/ModelingToolsEditorMode/Source/ModelingToolsEditorMode/Private/ModelingToolsEditorMode.cpp
@@ -1669,6 +1669,8 @@ void UModelingToolsEditorMode::OnToolEnded(UInteractiveToolManager* Manager, UIn
 
 void UModelingToolsEditorMode::BindCommands()
 {
+	Super::BindCommands();
+
 	const FModelingToolsManagerCommands& ToolManagerCommands = FModelingToolsManagerCommands::Get();
 	const TSharedRef<FUICommandList>& CommandList = Toolkit->GetToolkitCommands();
 
@@ -1833,29 +1835,20 @@ void UModelingToolsEditorMode::FocusCameraAtCursorHotkey()
 }
 
 
-bool UModelingToolsEditorMode::ComputeBoundingBoxForViewportFocus(AActor* Actor, UPrimitiveComponent* PrimitiveComponent, FBox& InOutBox) const
+FBox UModelingToolsEditorMode::ComputeCustomViewportFocus() const
 {
+	// Modeling mode prefers a slightly farther-out focus
 	auto ProcessFocusBoxFunc = [](FBox& FocusBoxInOut)
 	{
 		double MaxDimension = FocusBoxInOut.GetExtent().GetMax();
-		double ExpandAmount = (MaxDimension > SMALL_NUMBER) ? (MaxDimension * 0.2) : 25;		// 25 is a bit arbitrary here...
 		FocusBoxInOut = FocusBoxInOut.ExpandBy(MaxDimension * 0.2);
 	};
 
-	// if Tool supports custom Focus box, use that
-	if (GetToolManager()->HasAnyActiveTool())
+	FBox FocusBox = Super::ComputeCustomViewportFocus();
+	if (FocusBox.IsValid)
 	{
-		UInteractiveTool* Tool = GetToolManager()->GetActiveTool(EToolSide::Mouse);
-		IInteractiveToolCameraFocusAPI* FocusAPI = Cast<IInteractiveToolCameraFocusAPI>(Tool);
-		if (FocusAPI && FocusAPI->SupportsWorldSpaceFocusBox() )
-		{
-			InOutBox = FocusAPI->GetWorldSpaceFocusBox();
-			if (InOutBox.IsValid)
-			{
-				ProcessFocusBoxFunc(InOutBox);
-				return true;
-			}
-		}
+		ProcessFocusBoxFunc(FocusBox);
+		return FocusBox;
 	}
 
 	// if we have an active Selection we can focus on that
@@ -1863,12 +1856,29 @@ bool UModelingToolsEditorMode::ComputeBoundingBoxForViewportFocus(AActor* Actor,
 	{
 		UE::Geometry::FGeometrySelectionBounds SelectionBounds;
 		GetSelectionManager()->GetSelectionBounds(SelectionBounds);
-		InOutBox = (FBox)SelectionBounds.WorldBounds;
-		ProcessFocusBoxFunc(InOutBox);
+		FocusBox = (FBox)SelectionBounds.WorldBounds;
+		ProcessFocusBoxFunc(FocusBox);
+		return FocusBox;
+	}
+
+	// did not set a focus box, return a default (invalid) box
+	return FBox();
+}
+
+bool UModelingToolsEditorMode::HasCustomViewportFocus() const
+{
+	if (Super::HasCustomViewportFocus())
+	{
+		return true;
+	}
+
+	// if we have an active Selection we can focus on that
+	if (GetSelectionManager() && GetSelectionManager()->HasSelection())
+	{
 		return true;
 	}
 
-	// fallback to base focus behavior
+	// no mode-specific focus behavior
 	return false;
 }
 
diff --git a/Engine/Plugins/Editor/ModelingToolsEditorMode/Source/ModelingToolsEditorMode/Public/ModelingToolsEditorMode.h b/Engine/Plugins/Editor/ModelingToolsEditorMode/Source/ModelingToolsEditorMode/Public/ModelingToolsEditorMode.h
--- a/Engine/Plugins/Editor/ModelingToolsEditorMode/Source/ModelingToolsEditorMode/Public/ModelingToolsEditorMode.h
+++ b/Engine/Plugins/Editor/ModelingToolsEditorMode/Source/ModelingToolsEditorMode/Public/ModelingToolsEditorMode.h
@@ -47,8 +47,6 @@ public:
 
 	virtual bool CanAutoSave() const override;
 
-	virtual bool ComputeBoundingBoxForViewportFocus(AActor* Actor, UPrimitiveComponent* PrimitiveComponent, FBox& InOutBox) const override;
-
 	virtual bool GetPivotForOrbit(FVector& OutPivot) const override;
 
 	/*
@@ -73,6 +71,11 @@ public:
 	virtual bool FrustumSelect(const FConvexVolume& InFrustum, FEditorViewportClient* InViewportClient, bool InSelect) override;
 
 
+	// Manage viewport focus
+	virtual bool HasCustomViewportFocus() const override;
+	virtual FBox ComputeCustomViewportFocus() const override;
+
+
 
 	//
 	// Selection System configuration, this will likely move elsewhere
diff --git a/Engine/Source/Editor/UnrealEd/Private/Tools/UEdMode.cpp b/Engine/Source/Editor/UnrealEd/Private/Tools/UEdMode.cpp
--- a/Engine/Source/Editor/UnrealEd/Private/Tools/UEdMode.cpp
+++ b/Engine/Source/Editor/UnrealEd/Private/Tools/UEdMode.cpp
@@ -8,11 +8,12 @@
 #include "EdModeInteractiveToolsContext.h"
 #include "Toolkits/ToolkitManager.h"
 #include "InteractiveToolManager.h"
+#include "InteractiveToolQueryInterfaces.h"
 #include "GameFramework/Actor.h"
 #include "Elements/Framework/TypedElementSelectionSet.h"
 #include "Elements/Interfaces/TypedElementObjectInterface.h"
 #include "Settings/LevelEditorViewportSettings.h"
-
+#include "EditorViewportCommands.h"
 
 //////////////////////////////////
 // UEdMode
@@ -280,3 +281,55 @@ bool UEdMode::IsSnapRotationEnabled()
 {
 	return GetDefault<ULevelEditorViewportSettings>()->RotGridEnabled;
 }
+
+void UEdMode::BindCommands()
+{
+	if (Toolkit)
+	{
+		const TSharedRef<FUICommandList>& CommandList = Toolkit->GetToolkitCommands();
+		const FEditorViewportCommands& ViewportCommands = FEditorViewportCommands::Get();
+		CommandList->MapAction(
+			ViewportCommands.FocusViewportToSelection,
+			FExecuteAction::CreateLambda([this]() 
+			{
+				FBox FocusBox = ComputeCustomViewportFocus();
+				if (FocusBox.IsValid)
+				{
+					GEditor->MoveViewportCamerasToBox(FocusBox, true);
+				}
+				}),
+			FCanExecuteAction::CreateLambda([this]() { return HasCustomViewportFocus(); })
+		);
+	}
+
+}
+
+FBox UEdMode::GetFocusBoxFromActiveToolFocusAPI() const
+{
+	UInteractiveToolManager* ToolManager = GetToolManager();
+	if (ToolManager && ToolManager->HasAnyActiveTool())
+	{
+		UInteractiveTool* Tool = ToolManager->GetActiveTool(EToolSide::Mouse);
+		IInteractiveToolCameraFocusAPI* FocusAPI = Cast<IInteractiveToolCameraFocusAPI>(Tool);
+		if (FocusAPI && FocusAPI->SupportsWorldSpaceFocusBox())
+		{
+			return FocusAPI->GetWorldSpaceFocusBox();
+		}
+	}
+	return FBox();
+}
+
+bool UEdMode::HasCustomViewportFocus() const
+{
+	// Support ITF Focus API
+	FBox FocusBox = GetFocusBoxFromActiveToolFocusAPI();
+	return (bool)FocusBox.IsValid;
+}
+
+FBox UEdMode::ComputeCustomViewportFocus() const
+{
+	// Support ITF Focus API
+	return GetFocusBoxFromActiveToolFocusAPI();
+}
+
+
diff --git a/Engine/Source/Editor/UnrealEd/Public/Tools/UEdMode.h b/Engine/Source/Editor/UnrealEd/Public/Tools/UEdMode.h
--- a/Engine/Source/Editor/UnrealEd/Public/Tools/UEdMode.h
+++ b/Engine/Source/Editor/UnrealEd/Public/Tools/UEdMode.h
@@ -169,7 +169,8 @@ public:
 	virtual void SelectionChanged() {}
 
 	/**
-	 * Allows an editor mode to override the bounding box used to focus the viewport on a selection
+	 * Allows an editor mode to override the bounding box used to focus the viewport on an editor selection.
+	 * Note: Not used when HasCustomViewportFocus() returns true.
 	 *
 	 * @param Actor			The selected actor that is being considered for focus
 	 * @param PrimitiveComponent	The component in the actor being considered for focus
@@ -178,6 +179,18 @@ public:
 	 */
 	virtual bool ComputeBoundingBoxForViewportFocus(AActor* Actor, UPrimitiveComponent* PrimitiveComponent, FBox& InOutBox) const { return false; }
 
+	// Allows an editor mode to override the viewport focus logic, regardless of editor selection.
+	// @return true if the mode wants to control viewport focus
+	UNREALED_API virtual bool HasCustomViewportFocus() const;
+
+	// If HasCustomViewportFocus() returns true, this is called to apply the desired viewport focus
+	UNREALED_API virtual FBox ComputeCustomViewportFocus() const;
+
+private:
+	// @return the active focus box from the active tool's ITF focus api, if applicable, or an invalid box otherwise
+	FBox GetFocusBoxFromActiveToolFocusAPI() const;
+public:
+
 	/** Handling SelectActor */
 	virtual bool Select(AActor* InActor, bool bInSelected) { return 0; }
 
@@ -265,7 +278,7 @@ protected:
 	virtual void OnToolStarted(UInteractiveToolManager* Manager, UInteractiveTool* Tool) {}
 	virtual void OnToolEnded(UInteractiveToolManager* Manager, UInteractiveTool* Tool) {}
 	virtual void ActivateDefaultTool() {}
-	virtual void BindCommands() {}
+	UNREALED_API virtual void BindCommands();
 	UNREALED_API void OnModeActivated(const FEditorModeID& InID, bool bIsActive);
 
 private:
