ModelingTools: Added/Updated icons for the following primitive tools: (Arrow, Box, Capsule, Cone, Cylinder, Sphere)

diff --git a/Engine/Plugins/Editor/ModelingToolsEditorMode/Source/ModelingToolsEditorMode/Private/ModelingToolsEditorModeStyle.cpp b/Engine/Plugins/Editor/ModelingToolsEditorMode/Source/ModelingToolsEditorMode/Private/ModelingToolsEditorModeStyle.cpp
--- a/Engine/Plugins/Editor/ModelingToolsEditorMode/Source/ModelingToolsEditorMode/Private/ModelingToolsEditorModeStyle.cpp
+++ b/Engine/Plugins/Editor/ModelingToolsEditorMode/Source/ModelingToolsEditorMode/Private/ModelingToolsEditorModeStyle.cpp
@@ -144,6 +144,8 @@ void FModelingToolsEditorModeStyle::Initialize()
 		// Modes Palette Toolbar Icons
 		StyleSet->Set("ModelingToolsManagerCommands.BeginAddBoxPrimitiveTool", 			new IMAGE_BRUSH_SVG("Icons/ModelingBox", 		Icon20x20));
 		StyleSet->Set("ModelingToolsManagerCommands.BeginAddBoxPrimitiveTool.Small", 		new IMAGE_BRUSH_SVG("Icons/ModelingBox",		Icon40x40));
+		StyleSet->Set("ModelingToolsManagerCommands.BeginAddCapsulePrimitiveTool",			new IMAGE_BRUSH_SVG("Icons/ModelingCapsule", Icon20x20));
+		StyleSet->Set("ModelingToolsManagerCommands.BeginAddCapsulePrimitiveTool.Small",	new IMAGE_BRUSH_SVG("Icons/ModelingCapsule", Icon20x20));
 		StyleSet->Set("ModelingToolsManagerCommands.BeginAddCylinderPrimitiveTool", 			new IMAGE_BRUSH_SVG("Icons/ModelingCylinder", 		Icon20x20));
 		StyleSet->Set("ModelingToolsManagerCommands.BeginAddCylinderPrimitiveTool.Small", 		new IMAGE_BRUSH_SVG("Icons/ModelingCylinder",		Icon40x40));
 		StyleSet->Set("ModelingToolsManagerCommands.BeginAddConePrimitiveTool", 			new IMAGE_BRUSH_SVG("Icons/ModelingCone", 		Icon20x20));
