ModelingTools: Enable tablet support for ModelingMode using the new StylusInput interfaces.

diff --git a/Engine/Plugins/Editor/ModelingToolsEditorMode/ModelingToolsEditorMode.uplugin b/Engine/Plugins/Editor/ModelingToolsEditorMode/ModelingToolsEditorMode.uplugin
--- a/Engine/Plugins/Editor/ModelingToolsEditorMode/ModelingToolsEditorMode.uplugin
+++ b/Engine/Plugins/Editor/ModelingToolsEditorMode/ModelingToolsEditorMode.uplugin
@@ -40,6 +40,11 @@
 			"Name": "ToolPresets",
 			"Enabled": true,
 			"TargetAllowList": [ "Editor" ]
+		},
+		{
+			"Name": "StylusInput",
+			"Enabled": true,
+			"TargetAllowList": [ "Editor" ]
 		}
 	]
 }
\ No newline at end of file
diff --git a/Engine/Plugins/Editor/ModelingToolsEditorMode/Source/ModelingToolsEditorMode/ModelingToolsEditorMode.Build.cs b/Engine/Plugins/Editor/ModelingToolsEditorMode/Source/ModelingToolsEditorMode/ModelingToolsEditorMode.Build.cs
--- a/Engine/Plugins/Editor/ModelingToolsEditorMode/Source/ModelingToolsEditorMode/ModelingToolsEditorMode.Build.cs
+++ b/Engine/Plugins/Editor/ModelingToolsEditorMode/Source/ModelingToolsEditorMode/ModelingToolsEditorMode.Build.cs
@@ -54,10 +54,12 @@ public class ModelingToolsEditorMode : ModuleRules
 				"ToolPresetAsset",
 				"ToolPresetEditor",
 				"EditorConfig",
+				"StylusInput",
 				// ... add private dependencies that you statically link with here ...	
 			}
 		);
 
 		PublicDefinitions.Add("WITH_PROXYLOD=" + (Target.Platform == UnrealTargetPlatform.Win64 ? '1' : '0'));
+		PrivateDefinitions.Add("ENABLE_STYLUS_SUPPORT=1");
 	}
 }
\ No newline at end of file
diff --git a/Engine/Plugins/Editor/ModelingToolsEditorMode/Source/ModelingToolsEditorMode/Private/ModelingStylusInputHandler.cpp b/Engine/Plugins/Editor/ModelingToolsEditorMode/Source/ModelingToolsEditorMode/Private/ModelingStylusInputHandler.cpp
new file mode 100644
--- /dev/null
+++ b/Engine/Plugins/Editor/ModelingToolsEditorMode/Source/ModelingToolsEditorMode/Private/ModelingStylusInputHandler.cpp
@@ -0,0 +1,94 @@
+﻿// Copyright Epic Games, Inc. All Rights Reserved.
+
+#include "ModelingStylusInputHandler.h"
+
+#if ENABLE_STYLUS_SUPPORT
+
+#include "StylusInputTabletContext.h"
+#include "Framework/Application/SlateApplication.h"
+
+using namespace UE::Modeling;
+using namespace UE::StylusInput;
+
+FStylusInputHandler::FStylusInputHandler()
+{	
+}
+
+FStylusInputHandler::~FStylusInputHandler()
+{
+	for (TPair<TSharedPtr<SWindow>, IStylusInputInstance*> Pair : StylusInputInstances)
+	{
+		IStylusInputInstance* InputInstance = Pair.Value;
+		InputInstance->RemoveEventHandler(this);
+		ReleaseInstance(InputInstance);
+	}
+}
+
+bool FStylusInputHandler::RegisterWindow(const TSharedRef<SWidget>& Widget)
+{
+	TSharedPtr<SWindow> Window = FSlateApplication::Get().FindWidgetWindow(Widget);
+	if (!Window)
+	{
+		return false;
+	}
+
+	if (StylusInputInstances.Contains(Window))
+	{
+		return false;
+	}
+
+	IStylusInputInstance* InputInstance = CreateInstance(*Window);
+	if (!InputInstance)
+	{
+		return false;
+	}
+	
+	InputInstance->AddEventHandler(this, EEventHandlerThread::OnGameThread);
+	StylusInputInstances.Add(Window, InputInstance);
+	return true;
+}
+
+FString FStylusInputHandler::GetName()
+{
+	return "ModelingStylusInputHandler";
+}
+
+void FStylusInputHandler::OnPacket(const FStylusInputPacket& Packet, IStylusInputInstance* Instance)
+{
+	if (Packet.Type != EPacketType::Invalid &&
+		Packet.Type != EPacketType::AboveDigitizer)
+	{
+		ProcessPacket(Packet, Instance);
+	}
+}
+
+void FStylusInputHandler::ProcessPacket(const FStylusInputPacket& Packet, IStylusInputInstance* Instance)
+{
+	const IStylusInputTabletContext* TabletContext = GetTabletContext(Instance, Packet.TabletContextID);
+	const ETabletSupportedProperties SupportedProperties = TabletContext->GetSupportedProperties();
+	const bool bSupportsNormalPressure = (SupportedProperties & ETabletSupportedProperties::NormalPressure) != ETabletSupportedProperties::None;
+	ActivePressure = bSupportsNormalPressure ? Packet.NormalPressure : 1.0f;
+}
+
+const IStylusInputTabletContext* FStylusInputHandler::GetTabletContext(IStylusInputInstance* Instance, uint32 TabletContextID)
+{
+	if (!Instance)
+	{
+		return nullptr;
+	}
+
+	const TSharedPtr<IStylusInputTabletContext>* TabletContext = TabletContexts.Find(TabletContextID);
+	if (!TabletContext)
+	{
+		if (const TSharedPtr<IStylusInputTabletContext>& NewTabletContext = Instance->GetTabletContext(TabletContextID))
+		{
+			// We currently assume that TabletContextIDs are unique across all instances.
+			TabletContext = &TabletContexts.Emplace(TabletContextID, NewTabletContext);
+		}
+	}
+
+	return TabletContext ? TabletContext->Get() : nullptr;
+}
+
+#endif // ENABLE_STYLUS_SUPPORT
+
diff --git a/Engine/Plugins/Editor/ModelingToolsEditorMode/Source/ModelingToolsEditorMode/Private/ModelingStylusInputHandler.h b/Engine/Plugins/Editor/ModelingToolsEditorMode/Source/ModelingToolsEditorMode/Private/ModelingStylusInputHandler.h
new file mode 100644
--- /dev/null
+++ b/Engine/Plugins/Editor/ModelingToolsEditorMode/Source/ModelingToolsEditorMode/Private/ModelingStylusInputHandler.h
@@ -0,0 +1,66 @@
+﻿// Copyright Epic Games, Inc. All Rights Reserved.
+#pragma once
+
+#if ENABLE_STYLUS_SUPPORT
+
+#include "CoreMinimal.h"
+#include "StylusInput.h"
+#include "StylusInputPacket.h"
+#include "TickableEditorObject.h"
+
+#include "BaseTools/MeshSurfacePointTool.h"
+#include "Containers/SpscQueue.h"
+#include "Widgets/SWidget.h"
+
+namespace UE::Modeling
+{
+
+/**
+ * FStylusInputHandler registers itself as a listener for stylus events and implements the
+ * IToolStylusStateProviderAPI interface, which allows MeshSurfacePointTool implementations
+ * to query for pen pressure.
+ *
+ * This is kind of a hack. Unfortunately, the current Stylus module is a Plugin so it
+ * cannot be used in the base ToolsFramework directly. So, we need this in the Mode as a
+ * workaround.
+ */
+class FStylusInputHandler final : public UE::StylusInput::IStylusInputEventHandler, public IToolStylusStateProviderAPI
+{
+public:
+	FStylusInputHandler();
+	virtual ~FStylusInputHandler() override;
+
+	/**
+	 * Registers the window containing the given Widget for Stylus input handling.
+	 * Registered windows persist for the lifetime of the StylusInputHandler.
+	 * 
+	 * @param Widget The widget whose window is to be registered.
+	 * @return True if the window was registered, false if the window was invalid or was previously registered.
+	 */
+	bool RegisterWindow(const TSharedRef<SWidget>& Widget);
+
+	
+	// IStylusInputEventHandler implementation
+	virtual FString GetName() override;
+	virtual void OnPacket(const UE::StylusInput::FStylusInputPacket& Packet, UE::StylusInput::IStylusInputInstance* Instance) override;
+
+	
+	// IToolStylusStateProviderAPI implementation
+	virtual float GetCurrentPressure() const override
+	{
+		return ActivePressure;
+	}
+
+private:
+	void ProcessPacket(const UE::StylusInput::FStylusInputPacket& Packet, UE::StylusInput::IStylusInputInstance* Instance);
+	const UE::StylusInput::IStylusInputTabletContext* GetTabletContext(UE::StylusInput::IStylusInputInstance* Instance, uint32 TabletContextID);
+
+	TMap<TSharedPtr<SWindow>, UE::StylusInput::IStylusInputInstance*> StylusInputInstances;
+	TMap<uint32, TSharedPtr<UE::StylusInput::IStylusInputTabletContext>> TabletContexts;
+
+	float ActivePressure = 1.0f;
+};
+
+}
+
+#endif // ENABLE_STYLUS_SUPPORT
diff --git a/Engine/Plugins/Editor/ModelingToolsEditorMode/Source/ModelingToolsEditorMode/Private/ModelingToolsEditorMode.cpp b/Engine/Plugins/Editor/ModelingToolsEditorMode/Source/ModelingToolsEditorMode/Private/ModelingToolsEditorMode.cpp
--- a/Engine/Plugins/Editor/ModelingToolsEditorMode/Source/ModelingToolsEditorMode/Private/ModelingToolsEditorMode.cpp
+++ b/Engine/Plugins/Editor/ModelingToolsEditorMode/Source/ModelingToolsEditorMode/Private/ModelingToolsEditorMode.cpp
@@ -122,15 +122,6 @@
 #include "EditorModeManager.h"
 #include "UnrealWidget.h"
 
-// Stylus support is currently disabled due to issues with the stylus plugin
-// We are leaving the code in this cpp file, defined out, so that it is easier to bring back if/when the stylus plugin is improved.
-#define ENABLE_STYLUS_SUPPORT 0
-
-#if ENABLE_STYLUS_SUPPORT 
-#include "IStylusState.h"
-#include "IStylusInputModule.h"
-#endif
-
 #include "LevelEditor.h"
 #include "SLevelViewport.h"
 #include "Application/ThrottleManager.h"
@@ -288,101 +279,6 @@ void UModelingToolsEditorMode::Tick(FEditorViewportClient* ViewportClient, float
 	}
 }
 
-// Note: Stylus support is currently non-functioning; the code to enable it is left here as reference in case it is brought back
-#if ENABLE_STYLUS_SUPPORT
-//
-// FStylusStateTracker registers itself as a listener for stylus events and implements
-// the IToolStylusStateProviderAPI interface, which allows MeshSurfacePointTool implementations
- // to query for the pen pressure.
-//
-// This is kind of a hack. Unfortunately the current Stylus module is a Plugin so it
-// cannot be used in the base ToolsFramework, and we need this in the Mode as a workaround.
-//
-class FStylusStateTracker : public IStylusMessageHandler, public IToolStylusStateProviderAPI
-{
-public:
-	const IStylusInputDevice* ActiveDevice = nullptr;
-	int32 ActiveDeviceIndex = -1;
-
-	bool bPenDown = false;
-	float ActivePressure = 1.0;
-
-	FStylusStateTracker()
-	{
-		UStylusInputSubsystem* StylusSubsystem = GEditor->GetEditorSubsystem<UStylusInputSubsystem>();
-		StylusSubsystem->AddMessageHandler(*this);
-
-		ActiveDevice = FindFirstPenDevice(StylusSubsystem, ActiveDeviceIndex);
-		bPenDown = false;
-	}
-
-	virtual ~FStylusStateTracker()
-	{
-		if (GEditor)
-		{
-			if (UStylusInputSubsystem* StylusSubsystem = GEditor->GetEditorSubsystem<UStylusInputSubsystem>())
-			{
-				StylusSubsystem->RemoveMessageHandler(*this);
-			}
-		}
-	}
-
-	virtual void OnStylusStateChanged(const FStylusState& NewState, int32 StylusIndex) override
-	{
-		if (ActiveDevice == nullptr)
-		{
-			UStylusInputSubsystem* StylusSubsystem = GEditor->GetEditorSubsystem<UStylusInputSubsystem>();
-			ActiveDevice = FindFirstPenDevice(StylusSubsystem, ActiveDeviceIndex);
-			bPenDown = false;
-		}
-		if (ActiveDevice != nullptr && ActiveDeviceIndex == StylusIndex)
-		{
-			bPenDown = NewState.IsStylusDown();
-			ActivePressure = NewState.GetPressure();
-		}
-	}
-
-
-	bool HaveActiveStylusState() const
-	{
-		return ActiveDevice != nullptr && bPenDown;
-	}
-
-	static const IStylusInputDevice* FindFirstPenDevice(const UStylusInputSubsystem* StylusSubsystem, int32& ActiveDeviceOut)
-	{
-		int32 NumDevices = StylusSubsystem->NumInputDevices();
-		for (int32 k = 0; k < NumDevices; ++k)
-		{
-			const IStylusInputDevice* Device = StylusSubsystem->GetInputDevice(k);
-			const TArray<EStylusInputType>& Inputs = Device->GetSupportedInputs();
-			for (EStylusInputType Input : Inputs)
-			{
-				if (Input == EStylusInputType::Pressure)
-				{
-					ActiveDeviceOut = k;
-					return Device;
-				}
-			}
-		}
-		return nullptr;
-	}
-
-
-
-	// IToolStylusStateProviderAPI implementation
-	virtual float GetCurrentPressure() const override
-	{
-		return (ActiveDevice != nullptr && bPenDown) ? ActivePressure : 1.0f;
-	}
-
-};
-#endif // ENABLE_STYLUS_SUPPORT
-
-
-
-
-
-
 void UModelingToolsEditorMode::Enter()
 {
 	UEdMode::Enter();
@@ -417,10 +313,12 @@ void UModelingToolsEditorMode::Enter()
 	GetInteractiveToolsContext()->OnRender.AddUObject(this, &UModelingToolsEditorMode::OnToolsContextRender);
 	GetInteractiveToolsContext()->OnDrawHUD.AddUObject(this, &UModelingToolsEditorMode::OnToolsContextDrawHUD);
 
-#if ENABLE_STYLUS_SUPPORT 
 	// register stylus event handler
-	StylusStateTracker = MakeUnique<FStylusStateTracker>();
-#endif
+	IToolStylusStateProviderAPI* StylusAPI = nullptr;
+	if (ensure(Toolkit.IsValid()))
+	{
+		StylusAPI = ((FModelingToolsEditorModeToolkit*)Toolkit.Get())->GetStylusStateProviderAPI();
+	}
 
 	// register gizmo helper
 	UE::TransformGizmoUtil::RegisterTransformGizmoContextObject(GetInteractiveToolsContext());
@@ -637,15 +535,11 @@ void UModelingToolsEditorMode::Enter()
 	//
 
 	auto MoveVerticesToolBuilder = NewObject<UMeshVertexSculptToolBuilder>();
-#if ENABLE_STYLUS_SUPPORT 
-	MoveVerticesToolBuilder->StylusAPI = StylusStateTracker.Get();
-#endif
+	MoveVerticesToolBuilder->StylusAPI = StylusAPI;
 	RegisterTool(ToolManagerCommands.BeginSculptMeshTool, TEXT("BeginSculptMeshTool"), MoveVerticesToolBuilder);
 
 	auto MeshGroupPaintToolBuilder = NewObject<UMeshGroupPaintToolBuilder>();
-#if ENABLE_STYLUS_SUPPORT 
-	MeshGroupPaintToolBuilder->StylusAPI = StylusStateTracker.Get();
-#endif
+	MeshGroupPaintToolBuilder->StylusAPI = StylusAPI;
 	RegisterTool(ToolManagerCommands.BeginMeshGroupPaintTool, TEXT("BeginMeshGroupPaintTool"), MeshGroupPaintToolBuilder);
 	RegisterTool(ToolManagerCommands.BeginMeshVertexPaintTool, TEXT("BeginMeshVertexPaintTool"), NewObject<UMeshVertexPaintToolBuilder>());
 
@@ -699,9 +593,7 @@ void UModelingToolsEditorMode::Enter()
 
 	auto DynaSculptToolBuilder = NewObject<UDynamicMeshSculptToolBuilder>();
 	DynaSculptToolBuilder->bEnableRemeshing = true;
-#if ENABLE_STYLUS_SUPPORT 
-	DynaSculptToolBuilder->StylusAPI = StylusStateTracker.Get();
-#endif
+	DynaSculptToolBuilder->StylusAPI = StylusAPI;
 	RegisterTool(ToolManagerCommands.BeginRemeshSculptMeshTool, TEXT("BeginRemeshSculptMeshTool"), DynaSculptToolBuilder);
 
 	RegisterTool(ToolManagerCommands.BeginRemeshMeshTool, TEXT("BeginRemeshMeshTool"), NewObject<URemeshMeshToolBuilder>());
@@ -1271,10 +1163,6 @@ void UModelingToolsEditorMode::Exit()
 		FEngineAnalytics::GetProvider().RecordEvent(TEXT("Editor.Usage.MeshModelingMode.Exit"), Attributes);
 	}
 
-#if ENABLE_STYLUS_SUPPORT 
-	StylusStateTracker = nullptr;
-#endif
-
 	UModelingToolsHostCustomizationAPI::Deregister(GetInteractiveToolsContext());
 
 	// TODO: cannot deregister currently because if another mode is also registering, its Enter()
diff --git a/Engine/Plugins/Editor/ModelingToolsEditorMode/Source/ModelingToolsEditorMode/Private/ModelingToolsEditorModeToolkit.cpp b/Engine/Plugins/Editor/ModelingToolsEditorMode/Source/ModelingToolsEditorMode/Private/ModelingToolsEditorModeToolkit.cpp
--- a/Engine/Plugins/Editor/ModelingToolsEditorMode/Source/ModelingToolsEditorMode/Private/ModelingToolsEditorModeToolkit.cpp
+++ b/Engine/Plugins/Editor/ModelingToolsEditorMode/Source/ModelingToolsEditorMode/Private/ModelingToolsEditorModeToolkit.cpp
@@ -66,7 +66,11 @@
 #include "ToolPresetAssetSubsystem.h"
 #include "Selection/GeometrySelectionManager.h"
 #include "Widgets/Input/SMultiLineEditableTextBox.h"
+#include "IAssetViewport.h"
 
+#if ENABLE_STYLUS_SUPPORT
+#include "ModelingStylusInputHandler.h"
+#endif
 
 #define LOCTEXT_NAMESPACE "FModelingToolsEditorModeToolkit"
 
@@ -170,6 +174,10 @@ FModelingToolsEditorModeToolkit::FModelingToolsEditorModeToolkit()
 
 	RecentPresetCollectionProvider = MakeShared< FRecentPresetCollectionProvider>();
 	CurrentPreset = MakeShared<FAssetData>();
+	
+#if ENABLE_STYLUS_SUPPORT
+	StylusInputHandler = MakeUnique<UE::Modeling::FStylusInputHandler>();
+#endif
 }
 
 FModelingToolsEditorModeToolkit::~FModelingToolsEditorModeToolkit()
@@ -2051,6 +2059,11 @@ void FModelingToolsEditorModeToolkit::InvokeUI()
 	// of FModeToolkit::InlineContentHolder so we can just replace it here
 	InlineContentHolder->SetContent(GetInlineContent().ToSharedRef());
 
+#if ENABLE_STYLUS_SUPPORT
+	// The ToolkitWidget is only attached to a valid window from this point onwards.
+	StylusInputHandler->RegisterWindow(ToolkitWidget.ToSharedRef());
+#endif
+
 	// if the ToolkitBuilder is being used, it will make up the UI
 	if (HasToolkitBuilder())
 	{
@@ -2393,6 +2406,13 @@ void FModelingToolsEditorModeToolkit::OnActiveViewportChanged(TSharedPtr<IAssetV
 			GetToolkitHost()->AddViewportOverlayWidget(SelectionPaletteOverlayWidget.ToSharedRef(), NewViewport);
 		}
 	}
+
+#if ENABLE_STYLUS_SUPPORT
+	if (StylusInputHandler)
+	{
+		StylusInputHandler->RegisterWindow(NewViewport->AsWidget());
+	}
+#endif
 }
 
 
@@ -2677,4 +2697,10 @@ FReply FModelingToolsEditorModeToolkit::HandleCompleteClick()
 	return FReply::Handled();
 }
 
+IToolStylusStateProviderAPI* FModelingToolsEditorModeToolkit::GetStylusStateProviderAPI() const
+{
+	return StylusInputHandler.Get();
+}
+
+
 #undef LOCTEXT_NAMESPACE
diff --git a/Engine/Plugins/Editor/ModelingToolsEditorMode/Source/ModelingToolsEditorMode/Public/ModelingToolsEditorMode.h b/Engine/Plugins/Editor/ModelingToolsEditorMode/Source/ModelingToolsEditorMode/Public/ModelingToolsEditorMode.h
--- a/Engine/Plugins/Editor/ModelingToolsEditorMode/Source/ModelingToolsEditorMode/Public/ModelingToolsEditorMode.h
+++ b/Engine/Plugins/Editor/ModelingToolsEditorMode/Source/ModelingToolsEditorMode/Public/ModelingToolsEditorMode.h
@@ -15,7 +15,6 @@ struct FToolBuilderState;
 
 class FEditorComponentSourceFactory;
 class FUICommandList;
-class FStylusStateTracker;		// for stylus events
 class FLevelObjectsObserver;
 class UModelingSceneSnappingManager;
 class UModelingSelectionInteraction;
@@ -118,9 +117,6 @@ protected:
 	FDelegateHandle EditorClosedEventHandle;
 	void OnEditorClosed();
 
-	// Stylus support is currently disabled; this is left in for reference if/when it is brought back
-	//TUniquePtr<FStylusStateTracker> StylusStateTracker;
-
 	TSharedPtr<FLevelObjectsObserver> LevelObjectsObserver;
 
 	UPROPERTY()
diff --git a/Engine/Plugins/Editor/ModelingToolsEditorMode/Source/ModelingToolsEditorMode/Public/ModelingToolsEditorModeToolkit.h b/Engine/Plugins/Editor/ModelingToolsEditorMode/Source/ModelingToolsEditorMode/Public/ModelingToolsEditorModeToolkit.h
--- a/Engine/Plugins/Editor/ModelingToolsEditorMode/Source/ModelingToolsEditorMode/Public/ModelingToolsEditorModeToolkit.h
+++ b/Engine/Plugins/Editor/ModelingToolsEditorMode/Source/ModelingToolsEditorMode/Public/ModelingToolsEditorModeToolkit.h
@@ -21,6 +21,12 @@ class UGeometrySelectionManager;
 class UInteractiveToolsPresetCollectionAsset;
 class FRecentPresetCollectionProvider;
 struct FAssetData;
+class IToolStylusStateProviderAPI;
+
+namespace UE::Modeling
+{
+	class FStylusInputHandler;
+}
 
 struct FToolPresetOption
 {
@@ -106,6 +112,8 @@ public:
 
 	void NotifySelectionSystemEnabledStateModified();
 
+	IToolStylusStateProviderAPI* GetStylusStateProviderAPI() const;
+
 private:
 	const static TArray<FName> PaletteNames_Standard;
 
@@ -216,6 +224,8 @@ private:
 	bool bFirstInitializeAfterModeSetup = true;
 	bool bShowActiveSelectionActions = true;
 
+	TUniquePtr<UE::Modeling::FStylusInputHandler> StylusInputHandler;
+
 };
 
 #undef LOCTEXT_NAMESPACE
