MeshModelingTools: Add icon for UV Transfer tool.

diff --git a/Engine/Plugins/Editor/ModelingToolsEditorMode/Source/ModelingToolsEditorMode/Private/ModelingToolsEditorModeStyle.cpp b/Engine/Plugins/Editor/ModelingToolsEditorMode/Source/ModelingToolsEditorMode/Private/ModelingToolsEditorModeStyle.cpp
--- a/Engine/Plugins/Editor/ModelingToolsEditorMode/Source/ModelingToolsEditorMode/Private/ModelingToolsEditorModeStyle.cpp
+++ b/Engine/Plugins/Editor/ModelingToolsEditorMode/Source/ModelingToolsEditorMode/Private/ModelingToolsEditorModeStyle.cpp
@@ -309,8 +309,8 @@ void FModelingToolsEditorModeStyle::Initialize()
 		StyleSet->Set("ModelingToolsManagerCommands.BeginTransformUVIslandsTool.Small",      new IMAGE_PLUGIN_BRUSH("Icons/TransformUVs_40x",     Icon20x20));         
 		StyleSet->Set("ModelingToolsManagerCommands.BeginUVLayoutTool",                      new IMAGE_PLUGIN_BRUSH("Icons/UVLayout_40x",         Icon20x20));    
 		StyleSet->Set("ModelingToolsManagerCommands.BeginUVLayoutTool.Small",                new IMAGE_PLUGIN_BRUSH("Icons/UVLayout_40x",         Icon20x20));  
-		StyleSet->Set("ModelingToolsManagerCommands.BeginUVTransferTool",                    new FSlateImageBrush(StyleSet->RootToCoreContentDir(TEXT("../Editor/Slate/Icons/GeneralTools/Next_40x.png")), Icon20x20));
-		StyleSet->Set("ModelingToolsManagerCommands.BeginUVTransferTool.Small",              new FSlateImageBrush(StyleSet->RootToCoreContentDir(TEXT("../Editor/Slate/Icons/GeneralTools/Next_40x.png")), Icon20x20));
+		StyleSet->Set("ModelingToolsManagerCommands.BeginUVTransferTool",                    new IMAGE_BRUSH_SVG("Icons/UVTransfer", Icon20x20));
+		StyleSet->Set("ModelingToolsManagerCommands.BeginUVTransferTool.Small",              new IMAGE_BRUSH_SVG("Icons/UVTransfer", Icon20x20));
 		StyleSet->Set("ModelingToolsManagerCommands.LaunchUVEditor",                         new IMAGE_BRUSH_SVG("Icons/UVEditor", Icon20x20));
 		StyleSet->Set("ModelingToolsManagerCommands.LaunchUVEditor.Small",                   new IMAGE_BRUSH_SVG("Icons/UVEditor", Icon20x20));
 
