MeshModelingTools: Rename Absolute Grid Snapping option in the gizmo settings.

diff --git a/Engine/Plugins/Editor/ModelingToolsEditorMode/Source/ModelingToolsEditorMode/Private/ModelingModeToolkit_Toolbars.cpp b/Engine/Plugins/Editor/ModelingToolsEditorMode/Source/ModelingToolsEditorMode/Private/ModelingModeToolkit_Toolbars.cpp
--- a/Engine/Plugins/Editor/ModelingToolsEditorMode/Source/ModelingToolsEditorMode/Private/ModelingModeToolkit_Toolbars.cpp
+++ b/Engine/Plugins/Editor/ModelingToolsEditorMode/Source/ModelingToolsEditorMode/Private/ModelingModeToolkit_Toolbars.cpp
@@ -113,8 +113,8 @@ void MakeSubMenu_GizmoVisibilityMode(FModelingToolsEditorModeToolkit* Toolkit, F
 		{
 			return Context->GetAbsoluteWorldSnappingEnabled();
 		}));
-	MenuBuilder.AddMenuEntry(LOCTEXT("GizmoMode_AbsoluteWorldSnap", "World Grid Snapping"), 
-		LOCTEXT("GizmoMode_AbsoluteWorldSnap_Tooltip", "Snap Translation/Rotation to Absolute Grid Coordinates in the World Coordinate System, instead of Relative to the initial position"),
+	MenuBuilder.AddMenuEntry(LOCTEXT("GizmoMode_AbsoluteGridSnap", "Absolute Grid Snapping"), 
+		LOCTEXT("GizmoMode_AbsoluteGridSnap_Tooltip", "Snap translation/rotation to absolute grid coordinates in the world coordinate system, instead of relative to the initial position"),
 		FSlateIcon(), GizmoMode_AbsoluteWorldSnap, NAME_None, EUserInterfaceActionType::ToggleButton);
 	
 
