ModelingTools: Disable stylus input support for non-Windows platforms to fix automation test errors.

diff --git a/Engine/Plugins/Editor/ModelingToolsEditorMode/Source/ModelingToolsEditorMode/ModelingToolsEditorMode.Build.cs b/Engine/Plugins/Editor/ModelingToolsEditorMode/Source/ModelingToolsEditorMode/ModelingToolsEditorMode.Build.cs
--- a/Engine/Plugins/Editor/ModelingToolsEditorMode/Source/ModelingToolsEditorMode/ModelingToolsEditorMode.Build.cs
+++ b/Engine/Plugins/Editor/ModelingToolsEditorMode/Source/ModelingToolsEditorMode/ModelingToolsEditorMode.Build.cs
@@ -60,6 +60,6 @@ public class ModelingToolsEditorMode : ModuleRules
 		);
 
 		PublicDefinitions.Add("WITH_PROXYLOD=" + (Target.Platform == UnrealTargetPlatform.Win64 ? '1' : '0'));
-		PrivateDefinitions.Add("ENABLE_STYLUS_SUPPORT=1");
+		PrivateDefinitions.Add("ENABLE_STYLUS_SUPPORT=" + (Target.Platform == UnrealTargetPlatform.Win64 ? '1' : '0'));
 	}
 }
\ No newline at end of file
diff --git a/Engine/Plugins/Editor/ModelingToolsEditorMode/Source/ModelingToolsEditorMode/Private/ModelingToolsEditorModeToolkit.cpp b/Engine/Plugins/Editor/ModelingToolsEditorMode/Source/ModelingToolsEditorMode/Private/ModelingToolsEditorModeToolkit.cpp
--- a/Engine/Plugins/Editor/ModelingToolsEditorMode/Source/ModelingToolsEditorMode/Private/ModelingToolsEditorModeToolkit.cpp
+++ b/Engine/Plugins/Editor/ModelingToolsEditorMode/Source/ModelingToolsEditorMode/Private/ModelingToolsEditorModeToolkit.cpp
@@ -2704,7 +2704,11 @@ FReply FModelingToolsEditorModeToolkit::HandleCompleteClick()
 
 IToolStylusStateProviderAPI* FModelingToolsEditorModeToolkit::GetStylusStateProviderAPI() const
 {
+#if ENABLE_STYLUS_SUPPORT
 	return StylusInputHandler.Get();
+#else
+	return nullptr;
+#endif
 }
 
 
diff --git a/Engine/Plugins/Editor/ModelingToolsEditorMode/Source/ModelingToolsEditorMode/Public/ModelingToolsEditorModeToolkit.h b/Engine/Plugins/Editor/ModelingToolsEditorMode/Source/ModelingToolsEditorMode/Public/ModelingToolsEditorModeToolkit.h
--- a/Engine/Plugins/Editor/ModelingToolsEditorMode/Source/ModelingToolsEditorMode/Public/ModelingToolsEditorModeToolkit.h
+++ b/Engine/Plugins/Editor/ModelingToolsEditorMode/Source/ModelingToolsEditorMode/Public/ModelingToolsEditorModeToolkit.h
@@ -23,10 +23,12 @@ class FRecentPresetCollectionProvider;
 struct FAssetData;
 class IToolStylusStateProviderAPI;
 
+#if ENABLE_STYLUS_SUPPORT
 namespace UE::Modeling
 {
 	class FStylusInputHandler;
 }
+#endif
 
 struct FToolPresetOption
 {
@@ -225,7 +227,9 @@ private:
 	bool bFirstInitializeAfterModeSetup = true;
 	bool bShowActiveSelectionActions = true;
 
+#if ENABLE_STYLUS_SUPPORT
 	TUniquePtr<UE::Modeling::FStylusInputHandler> StylusInputHandler;
+#endif
 
 };
 
