Fix Tri Select tool being called Triangle Edit in the Select tool group, so it has a consistent name and is not confused for the Triangle Edit tool in the Mesh tool group

diff --git a/Engine/Plugins/Editor/ModelingToolsEditorMode/Source/ModelingToolsEditorMode/Private/ModelingToolsManagerActions.cpp b/Engine/Plugins/Editor/ModelingToolsEditorMode/Source/ModelingToolsEditorMode/Private/ModelingToolsManagerActions.cpp
--- a/Engine/Plugins/Editor/ModelingToolsEditorMode/Source/ModelingToolsEditorMode/Private/ModelingToolsManagerActions.cpp
+++ b/Engine/Plugins/Editor/ModelingToolsEditorMode/Source/ModelingToolsEditorMode/Private/ModelingToolsManagerActions.cpp
@@ -366,7 +366,7 @@ void FModelingToolsManagerCommands::RegisterCommands()
 	else
 	{
 		REGISTER_MODELING_TOOL_COMMAND(BeginPolyModelTool_PolyEd, "PolyGroup Edit", "Select / Edit the current Mesh via PolyGroups");
-		REGISTER_MODELING_TOOL_COMMAND(BeginPolyModelTool_TriSel, "Triangle Edit", "Select / Edit the current Mesh via Triangles");
+		REGISTER_MODELING_TOOL_COMMAND(BeginPolyModelTool_TriSel, "Tri Select", "Select / Edit the current Mesh triangles with a brush interface");
 	}
 
 
