MeshModelingTools: Fix many LOCTEXT collisions caused by tool short names, due to UI_COMMAND using the same key when registering the same command.

diff --git a/Engine/Plugins/Editor/ModelingToolsEditorMode/Source/ModelingToolsEditorMode/Private/ModelingToolsEditorModeToolkit.cpp b/Engine/Plugins/Editor/ModelingToolsEditorMode/Source/ModelingToolsEditorMode/Private/ModelingToolsEditorModeToolkit.cpp
--- a/Engine/Plugins/Editor/ModelingToolsEditorMode/Source/ModelingToolsEditorMode/Private/ModelingToolsEditorModeToolkit.cpp
+++ b/Engine/Plugins/Editor/ModelingToolsEditorMode/Source/ModelingToolsEditorMode/Private/ModelingToolsEditorModeToolkit.cpp
@@ -224,6 +224,11 @@ TSharedPtr<SWidget> FModelingToolsEditorModeToolkit::GetInlineContent() const
 
 void FModelingToolsEditorModeToolkit::RegisterPalettes()
 {
+	// Note, currently this code path does not use short names- those only get used if 
+	//  UISettings->bUseLegacyModelingPalette is true, inside BuildToolPalette. If we someday
+	//  do want to use short names here, FModelingToolsManagerCommands::GetCommandLabel()
+	//  can give those.
+
 	const FModelingToolsManagerCommands& Commands = FModelingToolsManagerCommands::Get();
 	const TSharedPtr<FUICommandList> CommandList = GetToolkitCommands();
 
@@ -1829,6 +1834,13 @@ void FModelingToolsEditorModeToolkit::BuildToolPalette(FName PaletteIndex, class
 	
 	const FModelingToolsManagerCommands& Commands = FModelingToolsManagerCommands::Get();
 	UModelingToolsModeCustomizationSettings* UISettings = GetMutableDefault<UModelingToolsModeCustomizationSettings>();
+	
+	bool bUseShortNames = UISettings->bUseLegacyModelingPalette;
+	auto AddButton = [&Commands, bUseShortNames, &ToolbarBuilder](const TSharedPtr<const FUICommandInfo>& Command)
+	{
+		ToolbarBuilder.AddToolBarButton(Command, NAME_None,
+			Commands.GetCommandLabel(Command, bUseShortNames));
+	};
 
 	if (PaletteIndex == ModelingFavoritesTabName)
 	{
@@ -1838,192 +1850,192 @@ void FModelingToolsEditorModeToolkit::BuildToolPalette(FName PaletteIndex, class
 		
 		for (const TSharedPtr<const FUICommandInfo>& ToolCommand : FavoriteCommands)
 		{
-			ToolbarBuilder.AddToolBarButton(ToolCommand);
+			AddButton(ToolCommand);
 		}
 	}
 	else if (PaletteIndex == SelectionActionsTabName)
 	{
-		ToolbarBuilder.AddToolBarButton(Commands.BeginSelectionAction_Delete);
-		//ToolbarBuilder.AddToolBarButton(Commands.BeginSelectionAction_Disconnect);		// disabled for 5.2, available via TriSel Tool
-		ToolbarBuilder.AddToolBarButton(Commands.BeginSelectionAction_Extrude);
-		ToolbarBuilder.AddToolBarButton(Commands.BeginSelectionAction_Offset);
+		AddButton(Commands.BeginSelectionAction_Delete);
+		//AddButton(Commands.BeginSelectionAction_Disconnect);		// disabled for 5.2, available via TriSel Tool
+		AddButton(Commands.BeginSelectionAction_Extrude);
+		AddButton(Commands.BeginSelectionAction_Offset);
 
-		ToolbarBuilder.AddToolBarButton(Commands.BeginPolyModelTool_ExtrudeEdges);
-		ToolbarBuilder.AddToolBarButton(Commands.BeginPolyModelTool_PushPull);
-		ToolbarBuilder.AddToolBarButton(Commands.BeginPolyModelTool_Inset);
-		ToolbarBuilder.AddToolBarButton(Commands.BeginPolyModelTool_Outset);
-		ToolbarBuilder.AddToolBarButton(Commands.BeginPolyModelTool_CutFaces);
-		ToolbarBuilder.AddToolBarButton(Commands.BeginPolyModelTool_Bevel);
+		AddButton(Commands.BeginPolyModelTool_ExtrudeEdges);
+		AddButton(Commands.BeginPolyModelTool_PushPull);
+		AddButton(Commands.BeginPolyModelTool_Inset);
+		AddButton(Commands.BeginPolyModelTool_Outset);
+		AddButton(Commands.BeginPolyModelTool_CutFaces);
+		AddButton(Commands.BeginPolyModelTool_Bevel);
 
-		ToolbarBuilder.AddToolBarButton(Commands.BeginPolyModelTool_InsertEdgeLoop);
-		ToolbarBuilder.AddToolBarButton(Commands.BeginSelectionAction_Retriangulate);
+		AddButton(Commands.BeginPolyModelTool_InsertEdgeLoop);
+		AddButton(Commands.BeginSelectionAction_Retriangulate);
 
-		ToolbarBuilder.AddToolBarButton(Commands.BeginPolyModelTool_PolyEd);
-		ToolbarBuilder.AddToolBarButton(Commands.BeginPolyModelTool_TriSel);
+		AddButton(Commands.BeginPolyModelTool_PolyEd);
+		AddButton(Commands.BeginPolyModelTool_TriSel);
 		
 	}
 	else if (PaletteIndex == PrimitiveTabName)
 	{
-		ToolbarBuilder.AddToolBarButton(Commands.BeginAddBoxPrimitiveTool);
-		ToolbarBuilder.AddToolBarButton(Commands.BeginAddSpherePrimitiveTool);
-		ToolbarBuilder.AddToolBarButton(Commands.BeginAddCylinderPrimitiveTool);
-		ToolbarBuilder.AddToolBarButton(Commands.BeginAddCapsulePrimitiveTool);
-		ToolbarBuilder.AddToolBarButton(Commands.BeginAddConePrimitiveTool);
-		ToolbarBuilder.AddToolBarButton(Commands.BeginAddTorusPrimitiveTool);
-		ToolbarBuilder.AddToolBarButton(Commands.BeginAddArrowPrimitiveTool);
-		ToolbarBuilder.AddToolBarButton(Commands.BeginAddRectanglePrimitiveTool);
-		ToolbarBuilder.AddToolBarButton(Commands.BeginAddDiscPrimitiveTool);
-		ToolbarBuilder.AddToolBarButton(Commands.BeginAddStairsPrimitiveTool);
+		AddButton(Commands.BeginAddBoxPrimitiveTool);
+		AddButton(Commands.BeginAddSpherePrimitiveTool);
+		AddButton(Commands.BeginAddCylinderPrimitiveTool);
+		AddButton(Commands.BeginAddCapsulePrimitiveTool);
+		AddButton(Commands.BeginAddConePrimitiveTool);
+		AddButton(Commands.BeginAddTorusPrimitiveTool);
+		AddButton(Commands.BeginAddArrowPrimitiveTool);
+		AddButton(Commands.BeginAddRectanglePrimitiveTool);
+		AddButton(Commands.BeginAddDiscPrimitiveTool);
+		AddButton(Commands.BeginAddStairsPrimitiveTool);
 	}
 	else if (PaletteIndex == CreateTabName)
 	{
-		ToolbarBuilder.AddToolBarButton(Commands.BeginDrawPolygonTool);
-		ToolbarBuilder.AddToolBarButton(Commands.BeginDrawPolyPathTool);
-		ToolbarBuilder.AddToolBarButton(Commands.BeginDrawAndRevolveTool);
-		ToolbarBuilder.AddToolBarButton(Commands.BeginRevolveSplineTool);
-		ToolbarBuilder.AddToolBarButton(Commands.BeginRevolveBoundaryTool);
-		ToolbarBuilder.AddToolBarButton(Commands.BeginCombineMeshesTool);
-		ToolbarBuilder.AddToolBarButton(Commands.BeginDuplicateMeshesTool);
-		ToolbarBuilder.AddToolBarButton(Commands.BeginPatternTool);
+		AddButton(Commands.BeginDrawPolygonTool);
+		AddButton(Commands.BeginDrawPolyPathTool);
+		AddButton(Commands.BeginDrawAndRevolveTool);
+		AddButton(Commands.BeginRevolveSplineTool);
+		AddButton(Commands.BeginRevolveBoundaryTool);
+		AddButton(Commands.BeginCombineMeshesTool);
+		AddButton(Commands.BeginDuplicateMeshesTool);
+		AddButton(Commands.BeginPatternTool);
 	}
 	else if (PaletteIndex == TransformTabName)
 	{
-		ToolbarBuilder.AddToolBarButton(Commands.BeginTransformMeshesTool);
-		ToolbarBuilder.AddToolBarButton(Commands.BeginAlignObjectsTool);
-		ToolbarBuilder.AddToolBarButton(Commands.BeginEditPivotTool);
-		ToolbarBuilder.AddToolBarButton(Commands.BeginAddPivotActorTool);
-		ToolbarBuilder.AddToolBarButton(Commands.BeginBakeTransformTool);
-		ToolbarBuilder.AddToolBarButton(Commands.BeginTransferMeshTool);
-		ToolbarBuilder.AddToolBarButton(Commands.BeginConvertMeshesTool);
-		ToolbarBuilder.AddToolBarButton(Commands.BeginSplitMeshesTool);
+		AddButton(Commands.BeginTransformMeshesTool);
+		AddButton(Commands.BeginAlignObjectsTool);
+		AddButton(Commands.BeginEditPivotTool);
+		AddButton(Commands.BeginAddPivotActorTool);
+		AddButton(Commands.BeginBakeTransformTool);
+		AddButton(Commands.BeginTransferMeshTool);
+		AddButton(Commands.BeginConvertMeshesTool);
+		AddButton(Commands.BeginSplitMeshesTool);
 	}
 	else if (PaletteIndex == DeformTabName)
 	{
-		ToolbarBuilder.AddToolBarButton(Commands.BeginSculptMeshTool);
-		ToolbarBuilder.AddToolBarButton(Commands.BeginRemeshSculptMeshTool);
-		ToolbarBuilder.AddToolBarButton(Commands.BeginSmoothMeshTool);
-		ToolbarBuilder.AddToolBarButton(Commands.BeginOffsetMeshTool);
-		ToolbarBuilder.AddToolBarButton(Commands.BeginMeshSpaceDeformerTool);
-		ToolbarBuilder.AddToolBarButton(Commands.BeginLatticeDeformerTool);
-		ToolbarBuilder.AddToolBarButton(Commands.BeginDisplaceMeshTool);
+		AddButton(Commands.BeginSculptMeshTool);
+		AddButton(Commands.BeginRemeshSculptMeshTool);
+		AddButton(Commands.BeginSmoothMeshTool);
+		AddButton(Commands.BeginOffsetMeshTool);
+		AddButton(Commands.BeginMeshSpaceDeformerTool);
+		AddButton(Commands.BeginLatticeDeformerTool);
+		AddButton(Commands.BeginDisplaceMeshTool);
 	}
 	else if (PaletteIndex == MeshProcessingTabName)
 	{
-		ToolbarBuilder.AddToolBarButton(Commands.BeginSimplifyMeshTool);
-		ToolbarBuilder.AddToolBarButton(Commands.BeginRemeshMeshTool);
-		ToolbarBuilder.AddToolBarButton(Commands.BeginWeldEdgesTool);
-		ToolbarBuilder.AddToolBarButton(Commands.BeginRemoveOccludedTrianglesTool);
-		ToolbarBuilder.AddToolBarButton(Commands.BeginSelfUnionTool);
-		ToolbarBuilder.AddToolBarButton(Commands.BeginProjectToTargetTool);
+		AddButton(Commands.BeginSimplifyMeshTool);
+		AddButton(Commands.BeginRemeshMeshTool);
+		AddButton(Commands.BeginWeldEdgesTool);
+		AddButton(Commands.BeginRemoveOccludedTrianglesTool);
+		AddButton(Commands.BeginSelfUnionTool);
+		AddButton(Commands.BeginProjectToTargetTool);
 	}
 	else if (PaletteIndex == LODToolsTabName)
 	{
-		ToolbarBuilder.AddToolBarButton(Commands.BeginLODManagerTool);
-		ToolbarBuilder.AddToolBarButton(Commands.BeginGenerateStaticMeshLODAssetTool);
-		ToolbarBuilder.AddToolBarButton(Commands.BeginISMEditorTool);
+		AddButton(Commands.BeginLODManagerTool);
+		AddButton(Commands.BeginGenerateStaticMeshLODAssetTool);
+		AddButton(Commands.BeginISMEditorTool);
 	}
 	else if (PaletteIndex == VoxToolsTabName)
 	{
-		ToolbarBuilder.AddToolBarButton(Commands.BeginVoxelSolidifyTool);
-		ToolbarBuilder.AddToolBarButton(Commands.BeginVoxelBlendTool);
-		ToolbarBuilder.AddToolBarButton(Commands.BeginVoxelMorphologyTool);
+		AddButton(Commands.BeginVoxelSolidifyTool);
+		AddButton(Commands.BeginVoxelBlendTool);
+		AddButton(Commands.BeginVoxelMorphologyTool);
 #if WITH_PROXYLOD
-		ToolbarBuilder.AddToolBarButton(Commands.BeginVoxelBooleanTool);
-		ToolbarBuilder.AddToolBarButton(Commands.BeginVoxelMergeTool);
+		AddButton(Commands.BeginVoxelBooleanTool);
+		AddButton(Commands.BeginVoxelMergeTool);
 #endif	// WITH_PROXYLOD
 	}
 	else if (PaletteIndex == TriModelingTabName)
 	{
-		ToolbarBuilder.AddToolBarButton(Commands.BeginMeshSelectionTool);
-		ToolbarBuilder.AddToolBarButton(Commands.BeginTriEditTool);
-		ToolbarBuilder.AddToolBarButton(Commands.BeginHoleFillTool);
-		ToolbarBuilder.AddToolBarButton(Commands.BeginMirrorTool);
-		ToolbarBuilder.AddToolBarButton(Commands.BeginPlaneCutTool);
-		ToolbarBuilder.AddToolBarButton(Commands.BeginPolygonCutTool);
-		ToolbarBuilder.AddToolBarButton(Commands.BeginMeshTrimTool);
+		AddButton(Commands.BeginMeshSelectionTool);
+		AddButton(Commands.BeginTriEditTool);
+		AddButton(Commands.BeginHoleFillTool);
+		AddButton(Commands.BeginMirrorTool);
+		AddButton(Commands.BeginPlaneCutTool);
+		AddButton(Commands.BeginPolygonCutTool);
+		AddButton(Commands.BeginMeshTrimTool);
 	}
 	else if (PaletteIndex == PolyModelingTabName)
 	{
-		ToolbarBuilder.AddToolBarButton(Commands.BeginPolyEditTool);
-		ToolbarBuilder.AddToolBarButton(Commands.BeginPolyDeformTool);
-		ToolbarBuilder.AddToolBarButton(Commands.BeginCubeGridTool);
-		ToolbarBuilder.AddToolBarButton(Commands.BeginMeshBooleanTool);
-		ToolbarBuilder.AddToolBarButton(Commands.BeginCutMeshWithMeshTool);
-		ToolbarBuilder.AddToolBarButton(Commands.BeginSubdividePolyTool);
+		AddButton(Commands.BeginPolyEditTool);
+		AddButton(Commands.BeginPolyDeformTool);
+		AddButton(Commands.BeginCubeGridTool);
+		AddButton(Commands.BeginMeshBooleanTool);
+		AddButton(Commands.BeginCutMeshWithMeshTool);
+		AddButton(Commands.BeginSubdividePolyTool);
 	}
 	else if (PaletteIndex == AttributesTabName)
 	{
-		ToolbarBuilder.AddToolBarButton(Commands.BeginMeshInspectorTool);
-		ToolbarBuilder.AddToolBarButton(Commands.BeginEditNormalsTool);
-		ToolbarBuilder.AddToolBarButton(Commands.BeginEditTangentsTool);
-		ToolbarBuilder.AddToolBarButton(Commands.BeginAttributeEditorTool);
-		ToolbarBuilder.AddToolBarButton(Commands.BeginPolyGroupsTool);
-		ToolbarBuilder.AddToolBarButton(Commands.BeginMeshGroupPaintTool);
-		ToolbarBuilder.AddToolBarButton(Commands.BeginMeshAttributePaintTool);
-		ToolbarBuilder.AddToolBarButton(Commands.BeginEditMeshMaterialsTool);
+		AddButton(Commands.BeginMeshInspectorTool);
+		AddButton(Commands.BeginEditNormalsTool);
+		AddButton(Commands.BeginEditTangentsTool);
+		AddButton(Commands.BeginAttributeEditorTool);
+		AddButton(Commands.BeginPolyGroupsTool);
+		AddButton(Commands.BeginMeshGroupPaintTool);
+		AddButton(Commands.BeginMeshAttributePaintTool);
+		AddButton(Commands.BeginEditMeshMaterialsTool);
 	} 
 	else if (PaletteIndex == BakingToolsTabName )
 	{
-		ToolbarBuilder.AddToolBarButton(Commands.BeginBakeMeshAttributeMapsTool);
-		ToolbarBuilder.AddToolBarButton(Commands.BeginBakeMultiMeshAttributeMapsTool);
-		ToolbarBuilder.AddToolBarButton(Commands.BeginBakeMeshAttributeVertexTool);
-		ToolbarBuilder.AddToolBarButton(Commands.BeginBakeRenderCaptureTool);
+		AddButton(Commands.BeginBakeMeshAttributeMapsTool);
+		AddButton(Commands.BeginBakeMultiMeshAttributeMapsTool);
+		AddButton(Commands.BeginBakeMeshAttributeVertexTool);
+		AddButton(Commands.BeginBakeRenderCaptureTool);
 	}
 	else if (PaletteIndex == UVTabName)
 	{
-		ToolbarBuilder.AddToolBarButton(Commands.BeginGlobalUVGenerateTool);
-		ToolbarBuilder.AddToolBarButton(Commands.BeginGroupUVGenerateTool);
-		ToolbarBuilder.AddToolBarButton(Commands.BeginUVProjectionTool);
-		ToolbarBuilder.AddToolBarButton(Commands.BeginUVSeamEditTool);
-		ToolbarBuilder.AddToolBarButton(Commands.BeginTransformUVIslandsTool);
-		ToolbarBuilder.AddToolBarButton(Commands.BeginUVLayoutTool);
-		ToolbarBuilder.AddToolBarButton(Commands.BeginUVTransferTool);
+		AddButton(Commands.BeginGlobalUVGenerateTool);
+		AddButton(Commands.BeginGroupUVGenerateTool);
+		AddButton(Commands.BeginUVProjectionTool);
+		AddButton(Commands.BeginUVSeamEditTool);
+		AddButton(Commands.BeginTransformUVIslandsTool);
+		AddButton(Commands.BeginUVLayoutTool);
+		AddButton(Commands.BeginUVTransferTool);
 
 		// Handle the inclusion of the optional UVEditor button if the UVEditor plugin has been found
 		if (IModularFeatures::Get().IsModularFeatureAvailable(IUVEditorModularFeature::GetModularFeatureName()))
 		{
-			ToolbarBuilder.AddToolBarButton(Commands.LaunchUVEditor);
+			AddButton(Commands.LaunchUVEditor);
 		}
 	}
 	else if (PaletteIndex == VolumesTabName)
 	{
-		ToolbarBuilder.AddToolBarButton(Commands.BeginVolumeToMeshTool);
-		ToolbarBuilder.AddToolBarButton(Commands.BeginMeshToVolumeTool);
+		AddButton(Commands.BeginVolumeToMeshTool);
+		AddButton(Commands.BeginMeshToVolumeTool);
 		ToolbarBuilder.AddSeparator();
 
 		// BSPConv is disabled in Restrictive Mode.
 		if (Commands.BeginBspConversionTool)
 		{
-			ToolbarBuilder.AddToolBarButton(Commands.BeginBspConversionTool);
+			AddButton(Commands.BeginBspConversionTool);
 			ToolbarBuilder.AddSeparator();
 		}
 
-		ToolbarBuilder.AddToolBarButton(Commands.BeginPhysicsInspectorTool);
-		ToolbarBuilder.AddToolBarButton(Commands.BeginSetCollisionGeometryTool);
-		//ToolbarBuilder.AddToolBarButton(Commands.BeginEditCollisionGeometryTool);
-		ToolbarBuilder.AddToolBarButton(Commands.BeginExtractCollisionGeometryTool);
+		AddButton(Commands.BeginPhysicsInspectorTool);
+		AddButton(Commands.BeginSetCollisionGeometryTool);
+		//AddButton(Commands.BeginEditCollisionGeometryTool);
+		AddButton(Commands.BeginExtractCollisionGeometryTool);
 	}
 	else if (PaletteIndex == PrototypesTabName)
 	{
-		ToolbarBuilder.AddToolBarButton(Commands.BeginAddPatchTool);
-		ToolbarBuilder.AddToolBarButton(Commands.BeginShapeSprayTool);
+		AddButton(Commands.BeginAddPatchTool);
+		AddButton(Commands.BeginShapeSprayTool);
 	}
 	//else if (PaletteIndex == PolyEditTabName)
 	//{
-	//	ToolbarBuilder.AddToolBarButton(Commands.BeginPolyModelTool_FaceSelect);
-	//	ToolbarBuilder.AddToolBarButton(Commands.BeginPolyModelTool_EdgeSelect);
-	//	ToolbarBuilder.AddToolBarButton(Commands.BeginPolyModelTool_VertexSelect);
-	//	ToolbarBuilder.AddToolBarButton(Commands.BeginPolyModelTool_AllSelect);
-	//	ToolbarBuilder.AddToolBarButton(Commands.BeginPolyModelTool_LoopSelect);
-	//	ToolbarBuilder.AddToolBarButton(Commands.BeginPolyModelTool_RingSelect);
+	//	AddButton(Commands.BeginPolyModelTool_FaceSelect);
+	//	AddButton(Commands.BeginPolyModelTool_EdgeSelect);
+	//	AddButton(Commands.BeginPolyModelTool_VertexSelect);
+	//	AddButton(Commands.BeginPolyModelTool_AllSelect);
+	//	AddButton(Commands.BeginPolyModelTool_LoopSelect);
+	//	AddButton(Commands.BeginPolyModelTool_RingSelect);
 	//	ToolbarBuilder.AddSeparator();
-	//	ToolbarBuilder.AddToolBarButton(Commands.BeginPolyModelTool_Extrude);
-	//	ToolbarBuilder.AddToolBarButton(Commands.BeginPolyModelTool_Inset);
-	//	ToolbarBuilder.AddToolBarButton(Commands.BeginPolyModelTool_Outset);
-	//	ToolbarBuilder.AddToolBarButton(Commands.BeginPolyModelTool_CutFaces);
+	//	AddButton(Commands.BeginPolyModelTool_Extrude);
+	//	AddButton(Commands.BeginPolyModelTool_Inset);
+	//	AddButton(Commands.BeginPolyModelTool_Outset);
+	//	AddButton(Commands.BeginPolyModelTool_CutFaces);
 	//	ToolbarBuilder.AddSeparator();
-	//	ToolbarBuilder.AddToolBarButton(Commands.BeginSubdividePolyTool);
-	//	ToolbarBuilder.AddToolBarButton(Commands.BeginPolyEditTool);
+	//	AddButton(Commands.BeginSubdividePolyTool);
+	//	AddButton(Commands.BeginPolyEditTool);
 	//}
 	else
 	{
@@ -2041,7 +2053,7 @@ void FModelingToolsEditorModeToolkit::BuildToolPalette(FName PaletteIndex, class
 				Extensions[k]->GetExtensionTools(ExtensionQueryInfo, ToolSet);
 				for (const FExtensionToolDescription& ToolInfo : ToolSet)
 				{
-					ToolbarBuilder.AddToolBarButton(ToolInfo.ToolCommand);
+					AddButton(ToolInfo.ToolCommand);
 				}
 			}
 		}
diff --git a/Engine/Plugins/Editor/ModelingToolsEditorMode/Source/ModelingToolsEditorMode/Private/ModelingToolsManagerActions.cpp b/Engine/Plugins/Editor/ModelingToolsEditorMode/Source/ModelingToolsEditorMode/Private/ModelingToolsManagerActions.cpp
--- a/Engine/Plugins/Editor/ModelingToolsEditorMode/Source/ModelingToolsEditorMode/Private/ModelingToolsManagerActions.cpp
+++ b/Engine/Plugins/Editor/ModelingToolsEditorMode/Source/ModelingToolsEditorMode/Private/ModelingToolsManagerActions.cpp
@@ -36,22 +36,44 @@ TSharedPtr<FUICommandInfo> FModelingToolsManagerCommands::FindToolByName(FString
 	return TSharedPtr<FUICommandInfo>();
 }
 
+FText FModelingToolsManagerCommands::GetCommandLabel(TSharedPtr<const FUICommandInfo> Command, bool bShortName) const
+{
+	if (!ensure(Command.IsValid()))
+	{
+		return FText();
+	}
+	if (const int32* Index = CommandToRegisteredToolsIndex.Find(Command.Get()))
+	{
+		if (RegisteredTools.IsValidIndex(*Index) && !RegisteredTools[*Index].ShortName.IsEmpty())
+		{
+			return RegisteredTools[*Index].ShortName;
+		}
+	}
+	return Command->GetLabel();
+}
+
 
 void FModelingToolsManagerCommands::RegisterCommands()
 {
 	const UModelingToolsEditorModeSettings* Settings = GetDefault<UModelingToolsEditorModeSettings>();
-	const UModelingToolsModeCustomizationSettings* UISettings = GetDefault<UModelingToolsModeCustomizationSettings>();
-	bool bUseShortNames = UISettings->bUseLegacyModelingPalette;
 
 	// this has to be done with a compile-time macro because UI_COMMAND expands to LOCTEXT macros
 #define REGISTER_MODELING_TOOL_COMMAND(ToolCommandInfo, ToolName, ToolTip) \
 		UI_COMMAND(ToolCommandInfo, ToolName, ToolTip, EUserInterfaceActionType::ToggleButton, FInputChord()); \
 		RegisteredTools.Add(FStartToolCommand{ ToolName, ToolCommandInfo });
+	// If using short names, we can't use the same command with a different ToolName, because that will
+	//  generate colliding LOCTEXT keys. Instead we keep the short name in a map.
+#define REGISTER_MODELING_TOOL_COMMAND_WITH_SHORTNAME(ToolCommandInfo, ToolName, ShortName, ToolTip) \
+		UI_COMMAND(ToolCommandInfo, ToolName, ToolTip, EUserInterfaceActionType::ToggleButton, FInputChord()); \
+		RegisteredTools.Add(FStartToolCommand{ ToolName, ToolCommandInfo, LOCTEXT(ToolName"_short", ShortName) });
 
 	// this has to be done with a compile-time macro because UI_COMMAND expands to LOCTEXT macros
 #define REGISTER_MODELING_ACTION_COMMAND(ToolCommandInfo, ToolName, ToolTip) \
 		UI_COMMAND(ToolCommandInfo, ToolName, ToolTip, EUserInterfaceActionType::Button, FInputChord()); \
 		RegisteredTools.Add(FStartToolCommand{ ToolName, ToolCommandInfo });
+#define REGISTER_MODELING_ACTION_COMMAND_WITH_SHORTNAME(ToolCommandInfo, ToolName, ShortName, ToolTip) \
+		UI_COMMAND(ToolCommandInfo, ToolName, ToolTip, EUserInterfaceActionType::Button, FInputChord()); \
+		RegisteredTools.Add(FStartToolCommand{ ToolName, ToolCommandInfo, LOCTEXT(ToolName"_short", ShortName) });
 
 	// Shapes
 	REGISTER_MODELING_TOOL_COMMAND(BeginAddBoxPrimitiveTool, "Box", "Create new box objects");
@@ -62,116 +84,52 @@ void FModelingToolsManagerCommands::RegisterCommands()
 	REGISTER_MODELING_TOOL_COMMAND(BeginAddDiscPrimitiveTool, "Disc", "Create new disc objects");
 	REGISTER_MODELING_TOOL_COMMAND(BeginAddStairsPrimitiveTool, "Stairs", "Create new stairs objects");
 	REGISTER_MODELING_TOOL_COMMAND(BeginAddCapsulePrimitiveTool, "Capsule", "Create new capsule objects");
-	if ( bUseShortNames )
-	{
-		REGISTER_MODELING_TOOL_COMMAND(BeginAddCylinderPrimitiveTool, "Cyl", "Create new cylinder objects");
-		REGISTER_MODELING_TOOL_COMMAND(BeginAddRectanglePrimitiveTool, "Rect", "Create new rectangle objects");
-	}
-	else
-	{
-		REGISTER_MODELING_TOOL_COMMAND(BeginAddCylinderPrimitiveTool, "Cylinder", "Create new cylinder objects");
-		REGISTER_MODELING_TOOL_COMMAND(BeginAddRectanglePrimitiveTool, "Rectangle", "Create new rectangle objects");
-	}
+	REGISTER_MODELING_TOOL_COMMAND_WITH_SHORTNAME(BeginAddCylinderPrimitiveTool, "Cylinder", "Cyl", "Create new cylinder objects");
+	REGISTER_MODELING_TOOL_COMMAND_WITH_SHORTNAME(BeginAddRectanglePrimitiveTool, "Rectangle", "Rect", "Create new rectangle objects");
 
 
 	// Create
-	if ( bUseShortNames )
-	{
-		REGISTER_MODELING_TOOL_COMMAND(BeginDrawPolygonTool, "PolyExt", "Draw and extrude 2D Polygons to create new objects");
-		REGISTER_MODELING_TOOL_COMMAND(BeginDrawPolyPathTool, "PathExt", "Draw and extrude 2D Paths to create new objects");
-		REGISTER_MODELING_TOOL_COMMAND(BeginDrawAndRevolveTool, "PathRev", "Draw and revolve 2D Paths to create new objects");
-		REGISTER_MODELING_TOOL_COMMAND(BeginRevolveSplineTool, "SplnRev", "Revolve splines to create new objects");
-		REGISTER_MODELING_TOOL_COMMAND(BeginRevolveBoundaryTool, "BdryRev", "Revolve Mesh boundary loops to create new objects");
-		REGISTER_MODELING_TOOL_COMMAND(BeginCombineMeshesTool, "Merge", "Merge multiple Meshes to create new objects");
-		REGISTER_MODELING_TOOL_COMMAND(BeginDuplicateMeshesTool, "Dupe", "Duplicate single Meshes to create new objects");
-		REGISTER_MODELING_TOOL_COMMAND(BeginDrawSplineTool, "DrwSpln", "Draw a spline in the level");
-		REGISTER_MODELING_TOOL_COMMAND(BeginTriangulateSplinesTool, "MshSpln", "Triangulate the Spline Components of selected actors to create new objects");
-	}
-	else
-	{
-		REGISTER_MODELING_TOOL_COMMAND(BeginDrawPolygonTool, "Extrude Polygon", "Draw and extrude 2D Polygons to create new objects");
-		REGISTER_MODELING_TOOL_COMMAND(BeginDrawPolyPathTool, "Extrude Path", "Draw and extrude 2D Paths to create new objects");
-		REGISTER_MODELING_TOOL_COMMAND(BeginDrawAndRevolveTool, "Revolve Path", "Draw and revolve 2D Paths to create new objects");
-		REGISTER_MODELING_TOOL_COMMAND(BeginRevolveSplineTool, "Revolve Spline", "Revolve splines to create new objects");
-		REGISTER_MODELING_TOOL_COMMAND(BeginRevolveBoundaryTool, "Revolve Boundary", "Revolve Mesh boundary loops to create new objects");
-		REGISTER_MODELING_TOOL_COMMAND(BeginCombineMeshesTool, "Merge", "Merge multiple Meshes to create new objects");
-		REGISTER_MODELING_TOOL_COMMAND(BeginDuplicateMeshesTool, "Duplicate", "Duplicate single Meshes to create new objects");
-		REGISTER_MODELING_TOOL_COMMAND(BeginDrawSplineTool, "Draw Spline", "Draw a spline in the level");
-		REGISTER_MODELING_TOOL_COMMAND(BeginTriangulateSplinesTool, "Mesh Splines", "Triangulate the Spline Components of selected actors to create new objects");
-	}
+	REGISTER_MODELING_TOOL_COMMAND_WITH_SHORTNAME(BeginDrawPolygonTool, "Extrude Polygon", "PolyExt", "Draw and extrude 2D Polygons to create new objects");
+	REGISTER_MODELING_TOOL_COMMAND_WITH_SHORTNAME(BeginDrawPolyPathTool, "Extrude Path", "PathExt", "Draw and extrude 2D Paths to create new objects");
+	REGISTER_MODELING_TOOL_COMMAND_WITH_SHORTNAME(BeginDrawAndRevolveTool, "Revolve Path", "PathRev", "Draw and revolve 2D Paths to create new objects");
+	REGISTER_MODELING_TOOL_COMMAND_WITH_SHORTNAME(BeginRevolveSplineTool, "Revolve Spline", "SplnRev", "Revolve splines to create new objects");
+	REGISTER_MODELING_TOOL_COMMAND_WITH_SHORTNAME(BeginRevolveBoundaryTool, "Revolve Boundary", "BdryRev", "Revolve Mesh boundary loops to create new objects");
+	REGISTER_MODELING_TOOL_COMMAND(BeginCombineMeshesTool, "Merge", "Merge multiple Meshes to create new objects");
+	REGISTER_MODELING_TOOL_COMMAND_WITH_SHORTNAME(BeginDuplicateMeshesTool, "Duplicate", "Dupe", "Duplicate single Meshes to create new objects");
+	REGISTER_MODELING_TOOL_COMMAND_WITH_SHORTNAME(BeginDrawSplineTool, "Draw Spline", "DrwSpln", "Draw a spline in the level");
+	REGISTER_MODELING_TOOL_COMMAND_WITH_SHORTNAME(BeginTriangulateSplinesTool, "Mesh Splines", "MshSpln", "Triangulate the Spline Components of selected actors to create new objects");
 
 	// PolyModel
-	if ( bUseShortNames )
-	{
-		REGISTER_MODELING_TOOL_COMMAND(BeginPolyEditTool, "PolyEd", "Edit Meshes via PolyGroups");
-		REGISTER_MODELING_TOOL_COMMAND(BeginPolyDeformTool, "PolyDef", "Deform Meshes via PolyGroups");
-		REGISTER_MODELING_TOOL_COMMAND(BeginCubeGridTool, "CubeGr", "Create block out Meshes using a repositionable grid");
-		REGISTER_MODELING_TOOL_COMMAND(BeginMeshBooleanTool, "MshBool", "Apply Boolean operations to the two selected Meshes");
-		REGISTER_MODELING_TOOL_COMMAND(BeginCutMeshWithMeshTool, "MshCut", "Split one Mesh into parts using a second Mesh");
-		REGISTER_MODELING_TOOL_COMMAND(BeginSubdividePolyTool, "SubDiv", "Subdivide the selected Mesh via PolyGroups or triangles");
-	}
-	else
-	{
-		REGISTER_MODELING_TOOL_COMMAND(BeginPolyEditTool, "PolyGroup Edit", "Edit Meshes via PolyGroups");
-		REGISTER_MODELING_TOOL_COMMAND(BeginPolyDeformTool, "Deform PolyGroups", "Deform Meshes via PolyGroups");
-		REGISTER_MODELING_TOOL_COMMAND(BeginCubeGridTool, "CubeGrid", "Create block out Meshes using a repositionable grid");
-		REGISTER_MODELING_TOOL_COMMAND(BeginMeshBooleanTool, "Boolean", "Apply Boolean operations to the two selected Meshes");
-		REGISTER_MODELING_TOOL_COMMAND(BeginCutMeshWithMeshTool, "Mesh Cut", "Split one Mesh into parts using a second Mesh");
-		REGISTER_MODELING_TOOL_COMMAND(BeginSubdividePolyTool, "Subdivide", "Subdivide the selected Mesh via PolyGroups or triangles");
-	}
+	REGISTER_MODELING_TOOL_COMMAND_WITH_SHORTNAME(BeginPolyEditTool, "PolyGroup Edit", "PolyEd", "Edit Meshes via PolyGroups");
+	REGISTER_MODELING_TOOL_COMMAND_WITH_SHORTNAME(BeginPolyDeformTool, "Deform PolyGroups", "PolyDef", "Deform Meshes via PolyGroups");
+	REGISTER_MODELING_TOOL_COMMAND_WITH_SHORTNAME(BeginCubeGridTool, "CubeGrid", "CubeGr", "Create block out Meshes using a repositionable grid");
+	REGISTER_MODELING_TOOL_COMMAND_WITH_SHORTNAME(BeginMeshBooleanTool, "Boolean", "MshBool", "Apply Boolean operations to the two selected Meshes");
+	REGISTER_MODELING_TOOL_COMMAND_WITH_SHORTNAME(BeginCutMeshWithMeshTool, "Mesh Cut", "MshCut", "Split one Mesh into parts using a second Mesh");
+	REGISTER_MODELING_TOOL_COMMAND_WITH_SHORTNAME(BeginSubdividePolyTool, "Subdivide", "SubDiv", "Subdivide the selected Mesh via PolyGroups or triangles");
 
 	// TriModel
-	if ( bUseShortNames )
-	{
-		REGISTER_MODELING_TOOL_COMMAND(BeginMeshSelectionTool, "TriSel", "Select and edit Mesh triangles with a brush interface");
-		REGISTER_MODELING_TOOL_COMMAND(BeginTriEditTool, "TriEd", "Select and Edit the Mesh vertices, edges, and triangles");
-		REGISTER_MODELING_TOOL_COMMAND(BeginHoleFillTool, "HFill", "Fill holes in the selected Mesh");
-		REGISTER_MODELING_TOOL_COMMAND(BeginPlaneCutTool, "PlnCut", "Cut the selected Meshes with a 3D plane");
-	}
-	else
-	{
-		REGISTER_MODELING_TOOL_COMMAND(BeginMeshSelectionTool, "Tri Select", "Select and edit Mesh triangles with a brush interface");
-		REGISTER_MODELING_TOOL_COMMAND(BeginTriEditTool, "Triangle Edit", "Select and Edit the Mesh vertices, edges, and triangles");
-		REGISTER_MODELING_TOOL_COMMAND(BeginHoleFillTool, "Fill Holes", "Fill holes in the selected Mesh");
-		REGISTER_MODELING_TOOL_COMMAND(BeginPlaneCutTool, "Plane Cut", "Cut the selected Meshes with a 3D plane");
-	}
+	REGISTER_MODELING_TOOL_COMMAND_WITH_SHORTNAME(BeginMeshSelectionTool, "Tri Select", "TriSel", "Select and edit Mesh triangles with a brush interface");
+	REGISTER_MODELING_TOOL_COMMAND_WITH_SHORTNAME(BeginTriEditTool, "Triangle Edit", "TriEd", "Select and Edit the Mesh vertices, edges, and triangles");
+	REGISTER_MODELING_TOOL_COMMAND_WITH_SHORTNAME(BeginHoleFillTool, "Fill Holes", "HFill", "Fill holes in the selected Mesh");
+	REGISTER_MODELING_TOOL_COMMAND_WITH_SHORTNAME(BeginPlaneCutTool, "Plane Cut", "PlnCut", "Cut the selected Meshes with a 3D plane");
 	REGISTER_MODELING_TOOL_COMMAND(BeginMirrorTool, "Mirror", "Mirror the selected Meshes");
 	REGISTER_MODELING_TOOL_COMMAND(BeginPolygonCutTool, "PolyCut", "Cut the selected Mesh with an extruded polygon");
 	REGISTER_MODELING_TOOL_COMMAND(BeginMeshTrimTool, "Trim", "Trim/Cut the selected Mesh with the second selected Mesh");
 
 	// Deform
-	if ( bUseShortNames )
-	{
-		REGISTER_MODELING_TOOL_COMMAND(BeginSculptMeshTool, "VSclpt", "Vertex sculpting");
-		REGISTER_MODELING_TOOL_COMMAND(BeginRemeshSculptMeshTool, "DSclpt", "Dynamic mesh sculpting");
-		REGISTER_MODELING_TOOL_COMMAND(BeginDisplaceMeshTool, "Displce", "Tessellate and Displace the selected Mesh");
-	}
-	else
-	{
-		REGISTER_MODELING_TOOL_COMMAND(BeginSculptMeshTool, "Vertex Sculpt", "Vertex sculpting");
-		REGISTER_MODELING_TOOL_COMMAND(BeginRemeshSculptMeshTool, "Dynamic Sculpt", "Dynamic mesh sculpting");
-		REGISTER_MODELING_TOOL_COMMAND(BeginDisplaceMeshTool, "Displace", "Tessellate and Displace the selected Mesh");
-	}
+	REGISTER_MODELING_TOOL_COMMAND_WITH_SHORTNAME(BeginSculptMeshTool, "Vertex Sculpt", "VSclpt", "Vertex sculpting");
+	REGISTER_MODELING_TOOL_COMMAND_WITH_SHORTNAME(BeginRemeshSculptMeshTool, "Dynamic Sculpt", "DSclpt", "Dynamic mesh sculpting");
+	REGISTER_MODELING_TOOL_COMMAND_WITH_SHORTNAME(BeginDisplaceMeshTool, "Displace", "Displce", "Tessellate and Displace the selected Mesh");
 	REGISTER_MODELING_TOOL_COMMAND(BeginSmoothMeshTool, "Smooth", "Smooth the shape of the selected Mesh");
 	REGISTER_MODELING_TOOL_COMMAND(BeginOffsetMeshTool, "Offset", "Offset the surface of the selected Mesh");
 	REGISTER_MODELING_TOOL_COMMAND(BeginMeshSpaceDeformerTool, "Warp", "Reshape the selected Mesh using space deformers");
 	REGISTER_MODELING_TOOL_COMMAND(BeginLatticeDeformerTool, "Lattice", "Deform the selected Mesh using a 3D lattice/grid");
 
 	// Transform
-	if ( bUseShortNames )
-	{
-		REGISTER_MODELING_TOOL_COMMAND(BeginTransformMeshesTool, "XForm", "Transform the selected Meshes");
-		REGISTER_MODELING_TOOL_COMMAND(BeginEditPivotTool, "Pivot", "Edit the pivot points of the selected Meshes");
-		REGISTER_MODELING_TOOL_COMMAND(BeginAddPivotActorTool, "PivotAct", "Add actor to act as a pivot for child components");
-		REGISTER_MODELING_TOOL_COMMAND(BeginBakeTransformTool, "BakeRS", "Bake rotation and scale into the selected Meshes");
-	}
-	else
-	{
-		REGISTER_MODELING_TOOL_COMMAND(BeginTransformMeshesTool, "Transform", "Transform the selected Meshes");
-		REGISTER_MODELING_TOOL_COMMAND(BeginEditPivotTool, "Edit Pivot", "Edit the pivot points of the selected Meshes");
-		REGISTER_MODELING_TOOL_COMMAND(BeginAddPivotActorTool, "Pivot Actor", "Add actor to act as a pivot for child components");
-		REGISTER_MODELING_TOOL_COMMAND(BeginBakeTransformTool, "Bake Transform", "Bake rotation and scale into the selected Meshes");
-	}
+	REGISTER_MODELING_TOOL_COMMAND_WITH_SHORTNAME(BeginTransformMeshesTool, "Transform", "XForm", "Transform the selected Meshes");
+	REGISTER_MODELING_TOOL_COMMAND_WITH_SHORTNAME(BeginEditPivotTool, "Edit Pivot", "Pivot", "Edit the pivot points of the selected Meshes");
+	REGISTER_MODELING_TOOL_COMMAND_WITH_SHORTNAME(BeginAddPivotActorTool, "Pivot Actor", "PivotAct", "Add actor to act as a pivot for child components");
+	REGISTER_MODELING_TOOL_COMMAND_WITH_SHORTNAME(BeginBakeTransformTool, "Bake Transform", "BakeRS", "Bake rotation and scale into the selected Meshes");
 	REGISTER_MODELING_TOOL_COMMAND(BeginAlignObjectsTool, "Align", "Align the selected Objects");
 	REGISTER_MODELING_TOOL_COMMAND(BeginTransferMeshTool, "Transfer", "Copy the first selected Mesh to the second selected Mesh");
 	REGISTER_MODELING_TOOL_COMMAND(BeginConvertMeshesTool, "Convert", "Convert the selected Meshes to a different type of Mesh Object");
@@ -181,14 +139,7 @@ void FModelingToolsManagerCommands::RegisterCommands()
 	
 
 	// MeshOps
-	if ( bUseShortNames )
-	{
-		REGISTER_MODELING_TOOL_COMMAND(BeginSimplifyMeshTool, "Simplfy", "Simplify the selected Mesh");
-	}
-	else
-	{
-		REGISTER_MODELING_TOOL_COMMAND(BeginSimplifyMeshTool, "Simplify", "Simplify the selected Mesh");
-	}
+	REGISTER_MODELING_TOOL_COMMAND_WITH_SHORTNAME(BeginSimplifyMeshTool, "Simplify", "Simplfy", "Simplify the selected Mesh");
 	REGISTER_MODELING_TOOL_COMMAND(BeginRemeshMeshTool, "Remesh", "Re-triangulate the selected Mesh");
 	REGISTER_MODELING_TOOL_COMMAND(BeginWeldEdgesTool, "Weld", "Weld overlapping Mesh edges");
 	REGISTER_MODELING_TOOL_COMMAND(BeginRemoveOccludedTrianglesTool, "Jacket", "Remove hidden triangles from the selected Meshes");
@@ -197,135 +148,59 @@ void FModelingToolsManagerCommands::RegisterCommands()
 
 
 	// VoxOps
-	if ( bUseShortNames )
-	{
-		REGISTER_MODELING_TOOL_COMMAND(BeginVoxelSolidifyTool, "VoxWrap", "Wrap the selected Meshes using SDFs/voxels");
-		REGISTER_MODELING_TOOL_COMMAND(BeginVoxelBlendTool, "VoxBlnd", "Blend the selected Meshes using SDFs/voxels");
-		REGISTER_MODELING_TOOL_COMMAND(BeginVoxelMorphologyTool, "VoxMrph", "Offset/Inset the selected Meshes using SDFs/voxels");
+	REGISTER_MODELING_TOOL_COMMAND_WITH_SHORTNAME(BeginVoxelSolidifyTool, "Voxel Wrap", "VoxWrap", "Wrap the selected Meshes using SDFs/voxels");
+	REGISTER_MODELING_TOOL_COMMAND_WITH_SHORTNAME(BeginVoxelBlendTool, "Voxel Blend", "VoxBlnd", "Blend the selected Meshes using SDFs/voxels");
+	REGISTER_MODELING_TOOL_COMMAND_WITH_SHORTNAME(BeginVoxelMorphologyTool, "Voxel Offset", "VoxMrph", "Offset/Inset the selected Meshes using SDFs/voxels");
 #if WITH_PROXYLOD
-		// The ProxyLOD plugin is currently only available on Windows. Without it, the following tools do not work as expected.
-		REGISTER_MODELING_TOOL_COMMAND(BeginVoxelBooleanTool, "VoxBool", "Boolean the selected Meshes using SDFs/voxels");
-		REGISTER_MODELING_TOOL_COMMAND(BeginVoxelMergeTool, "VoxMrg", "Merge the selected Meshes using SDFs/voxels");
+	// The ProxyLOD plugin is currently only available on Windows. Without it, the following tools do not work as expected.
+	REGISTER_MODELING_TOOL_COMMAND_WITH_SHORTNAME(BeginVoxelBooleanTool, "Voxel Boolean", "VoxBool", "Boolean the selected Meshes using SDFs/voxels");
+	REGISTER_MODELING_TOOL_COMMAND_WITH_SHORTNAME(BeginVoxelMergeTool, "Voxel Merge", "VoxMrg", "Merge the selected Meshes using SDFs/voxels");
 #endif	// WITH_PROXYLOD
-	}
-	else
-	{
-		REGISTER_MODELING_TOOL_COMMAND(BeginVoxelSolidifyTool, "Voxel Wrap", "Wrap the selected Meshes using SDFs/voxels");
-		REGISTER_MODELING_TOOL_COMMAND(BeginVoxelBlendTool, "Voxel Blend", "Blend the selected Meshes using SDFs/voxels");
-		REGISTER_MODELING_TOOL_COMMAND(BeginVoxelMorphologyTool, "Voxel Offset", "Offset/Inset the selected Meshes using SDFs/voxels");
-#if WITH_PROXYLOD
-		// The ProxyLOD plugin is currently only available on Windows. Without it, the following tools do not work as expected.
-		REGISTER_MODELING_TOOL_COMMAND(BeginVoxelBooleanTool, "Voxel Boolean", "Boolean the selected Meshes using SDFs/voxels");
-		REGISTER_MODELING_TOOL_COMMAND(BeginVoxelMergeTool, "Voxel Merge", "Merge the selected Meshes using SDFs/voxels");
-#endif	// WITH_PROXYLOD
-	}
 
 	// Attributes
-	if ( bUseShortNames )
-	{
-		REGISTER_MODELING_TOOL_COMMAND(BeginMeshInspectorTool, "Inspct", "Inspect Mesh attributes");
-		REGISTER_MODELING_TOOL_COMMAND(BeginEditNormalsTool, "Nrmls", "Recompute or Repair Normals");
-		REGISTER_MODELING_TOOL_COMMAND(BeginEditTangentsTool, "Tngnts", "Recompute Tangents");
-		REGISTER_MODELING_TOOL_COMMAND(BeginAttributeEditorTool, "AttrEd", "Edit/configure Mesh attributes");
-		REGISTER_MODELING_TOOL_COMMAND(BeginPolyGroupsTool, "GrpGen", "Generate new PolyGroups");
-		REGISTER_MODELING_TOOL_COMMAND(BeginMeshGroupPaintTool, "GrpPnt", "Paint new PolyGroups");
-		REGISTER_MODELING_TOOL_COMMAND(BeginMeshVertexPaintTool, "VtxPnt", "Paint Mesh Vertices");
-		REGISTER_MODELING_TOOL_COMMAND(BeginMeshAttributePaintTool, "MapPnt", "Paint attribute maps");
-		REGISTER_MODELING_TOOL_COMMAND(BeginEditMeshMaterialsTool, "MatEd", "Assign materials to selected triangles");
-	}
-	else
-	{
-		REGISTER_MODELING_TOOL_COMMAND(BeginMeshInspectorTool, "Inspect", "Inspect Mesh attributes");
-		REGISTER_MODELING_TOOL_COMMAND(BeginEditNormalsTool, "Normals", "Recompute or Repair Normals");
-		REGISTER_MODELING_TOOL_COMMAND(BeginEditTangentsTool, "Tangents", "Recompute Tangents");
-		REGISTER_MODELING_TOOL_COMMAND(BeginPolyGroupsTool, "Generate PolyGroups", "Generate new PolyGroups");
-		REGISTER_MODELING_TOOL_COMMAND(BeginMeshGroupPaintTool, "Paint PolyGroups", "Paint new PolyGroups");
-		REGISTER_MODELING_TOOL_COMMAND(BeginMeshVertexPaintTool, "Paint Vertex Colors", "Paint Mesh Vertex Colors");
-		REGISTER_MODELING_TOOL_COMMAND(BeginAttributeEditorTool, "Edit Attributes", "Edit/configure Mesh attributes");
-		REGISTER_MODELING_TOOL_COMMAND(BeginMeshAttributePaintTool, "Paint Maps", "Paint attribute maps");
-		REGISTER_MODELING_TOOL_COMMAND(BeginEditMeshMaterialsTool, "Edit Materials", "Assign materials to selected triangles");
-	}
+	REGISTER_MODELING_TOOL_COMMAND_WITH_SHORTNAME(BeginMeshInspectorTool, "Inspect", "Inspct", "Inspect Mesh attributes");
+	REGISTER_MODELING_TOOL_COMMAND_WITH_SHORTNAME(BeginEditNormalsTool, "Normals", "Nrmls", "Recompute or Repair Normals");
+	REGISTER_MODELING_TOOL_COMMAND_WITH_SHORTNAME(BeginEditTangentsTool, "Tangents", "Tngnts", "Recompute Tangents");
+	REGISTER_MODELING_TOOL_COMMAND_WITH_SHORTNAME(BeginPolyGroupsTool, "Generate PolyGroups", "GrpGen", "Generate new PolyGroups");
+	REGISTER_MODELING_TOOL_COMMAND_WITH_SHORTNAME(BeginMeshGroupPaintTool, "Paint PolyGroups", "GrpPnt", "Paint new PolyGroups");
+	REGISTER_MODELING_TOOL_COMMAND_WITH_SHORTNAME(BeginMeshVertexPaintTool, "Paint Vertex Colors", "VtxPnt", "Paint Mesh Vertex Colors");
+	REGISTER_MODELING_TOOL_COMMAND_WITH_SHORTNAME(BeginAttributeEditorTool, "Edit Attributes", "AttrEd", "Edit/configure Mesh attributes");
+	REGISTER_MODELING_TOOL_COMMAND_WITH_SHORTNAME(BeginMeshAttributePaintTool, "Paint Maps", "MapPnt", "Paint attribute maps");
+	REGISTER_MODELING_TOOL_COMMAND_WITH_SHORTNAME(BeginEditMeshMaterialsTool, "Edit Materials", "MatEd", "Assign materials to selected triangles");
 
 	// UVs
 	REGISTER_MODELING_TOOL_COMMAND(BeginGlobalUVGenerateTool, "AutoUV", "Automatically unwrap and pack UVs");
 	// This is done directly, not with the REGISTER_ macro, since we don't want it added to the tool list or use a toggle button
 	UI_COMMAND(LaunchUVEditor, "UVEditor", "Launch UV asset editor", EUserInterfaceActionType::Button, FInputChord());
-	if ( bUseShortNames )
-	{
-		REGISTER_MODELING_TOOL_COMMAND(BeginGroupUVGenerateTool, "Unwrap", "Recompute UVs for existing UV islands or PolyGroups");
-		REGISTER_MODELING_TOOL_COMMAND(BeginUVProjectionTool, "Project", "Compute UVs via projecting to simple shapes");
-		REGISTER_MODELING_TOOL_COMMAND(BeginUVSeamEditTool, "SeamEd", "Add UV seams");
-		REGISTER_MODELING_TOOL_COMMAND(BeginTransformUVIslandsTool, "XFormUV", "Transform UV islands in UV space");
-		REGISTER_MODELING_TOOL_COMMAND(BeginUVLayoutTool, "Layout", "Transform and Repack existing UVs");
-		REGISTER_MODELING_TOOL_COMMAND(BeginUVTransferTool, "Transfer", "Transfer UVs from a simplified mesh to a higher resolution mesh");
-	}
-	else
-	{
-		REGISTER_MODELING_TOOL_COMMAND(BeginGroupUVGenerateTool, "UV Unwrap", "Recompute UVs for existing UV islands or PolyGroups");
-		REGISTER_MODELING_TOOL_COMMAND(BeginUVProjectionTool, "Project UVs", "Compute UVs via projecting to simple shapes");
-		REGISTER_MODELING_TOOL_COMMAND(BeginUVSeamEditTool, "Edit UV Seams", "Add UV seams");
-		REGISTER_MODELING_TOOL_COMMAND(BeginTransformUVIslandsTool, "Transform UVs", "Transform UV islands in UV space");
-		REGISTER_MODELING_TOOL_COMMAND(BeginUVLayoutTool, "Layout UVs", "Transform and Repack existing UVs");
-		REGISTER_MODELING_TOOL_COMMAND(BeginUVTransferTool, "Transfer UVs", "Transfer UVs from a simplified mesh to a higher resolution mesh");
-	}
+	REGISTER_MODELING_TOOL_COMMAND_WITH_SHORTNAME(BeginGroupUVGenerateTool, "UV Unwrap", "Unwrap", "Recompute UVs for existing UV islands or PolyGroups");
+	REGISTER_MODELING_TOOL_COMMAND_WITH_SHORTNAME(BeginUVProjectionTool, "Project UVs", "Project", "Compute UVs via projecting to simple shapes");
+	REGISTER_MODELING_TOOL_COMMAND_WITH_SHORTNAME(BeginUVSeamEditTool, "Edit UV Seams", "SeamEd", "Add UV seams");
+	REGISTER_MODELING_TOOL_COMMAND_WITH_SHORTNAME(BeginTransformUVIslandsTool, "Transform UVs", "XFormUV", "Transform UV islands in UV space");
+	REGISTER_MODELING_TOOL_COMMAND_WITH_SHORTNAME(BeginUVLayoutTool, "Layout UVs", "Layout", "Transform and Repack existing UVs");
+	REGISTER_MODELING_TOOL_COMMAND_WITH_SHORTNAME(BeginUVTransferTool, "Transfer UVs", "Transfer", "Transfer UVs from a simplified mesh to a higher resolution mesh");
 
 	// Baking
-	if ( bUseShortNames )
-	{
-		REGISTER_MODELING_TOOL_COMMAND(BeginBakeMeshAttributeMapsTool, "BakeTx", "Bake textures for a target Mesh");
-		REGISTER_MODELING_TOOL_COMMAND(BeginBakeMultiMeshAttributeMapsTool, "BakeAll", "Bake textures for a target Mesh from multiple source Meshes");
-		REGISTER_MODELING_TOOL_COMMAND(BeginBakeMeshAttributeVertexTool, "BakeVtx", "Bake vertex colors for a target Mesh");
-		REGISTER_MODELING_TOOL_COMMAND(BeginBakeRenderCaptureTool, "BakeRC", "Bake renders into new textures for a target Mesh from multiple source Meshes");
-	}
-	else
-	{
-		REGISTER_MODELING_TOOL_COMMAND(BeginBakeMeshAttributeMapsTool, "Bake Textures", "Bake textures for a target Mesh");
-		REGISTER_MODELING_TOOL_COMMAND(BeginBakeMultiMeshAttributeMapsTool, "Bake All", "Bake textures for a target Mesh from multiple source Meshes");
-		REGISTER_MODELING_TOOL_COMMAND(BeginBakeMeshAttributeVertexTool, "Bake Vertex Colors", "Bake vertex colors for a target Mesh");
-		REGISTER_MODELING_TOOL_COMMAND(BeginBakeRenderCaptureTool, "Bake RC", "Bake renders into new textures for a target Mesh from multiple source Meshes");
-	}
+	REGISTER_MODELING_TOOL_COMMAND_WITH_SHORTNAME(BeginBakeMeshAttributeMapsTool, "Bake Textures", "BakeTx", "Bake textures for a target Mesh");
+	REGISTER_MODELING_TOOL_COMMAND_WITH_SHORTNAME(BeginBakeMultiMeshAttributeMapsTool, "Bake All", "BakeAll", "Bake textures for a target Mesh from multiple source Meshes");
+	REGISTER_MODELING_TOOL_COMMAND_WITH_SHORTNAME(BeginBakeMeshAttributeVertexTool, "Bake Vertex Colors", "BakeVtx", "Bake vertex colors for a target Mesh");
+	REGISTER_MODELING_TOOL_COMMAND_WITH_SHORTNAME(BeginBakeRenderCaptureTool, "Bake RC", "BakeRC", "Bake renders into new textures for a target Mesh from multiple source Meshes");
 
 	// Volumes
-	if ( bUseShortNames )
-	{
-		REGISTER_MODELING_TOOL_COMMAND(BeginVolumeToMeshTool, "Vol2Msh", "Convert a Volume to a new Mesh Object");
-		REGISTER_MODELING_TOOL_COMMAND(BeginMeshToVolumeTool, "Msh2Vol", "Convert a Mesh to a Volume");
-		if (!Settings->InRestrictiveMode())
-		{
-			REGISTER_MODELING_TOOL_COMMAND(BeginBspConversionTool, "BSPConv", "Convert BSP to a new Mesh Object");
-		}
-		REGISTER_MODELING_TOOL_COMMAND(BeginPhysicsInspectorTool, "PInspct", "Inspect the physics/collision geometry for selected Meshes");
-		REGISTER_MODELING_TOOL_COMMAND(BeginSimpleCollisionEditorTool, "SCollEdit", "Edit the simple collision shapes for the selected Mesh");
-		REGISTER_MODELING_TOOL_COMMAND(BeginSetCollisionGeometryTool, "Msh2Coll", "Convert selected Meshes to Simple Collision Geometry (for last selected)");
-		REGISTER_MODELING_TOOL_COMMAND(BeginExtractCollisionGeometryTool, "Coll2Msh", "Convert Simple Collision Geometry to a new Mesh Object");
-	}
-	else
+	REGISTER_MODELING_TOOL_COMMAND_WITH_SHORTNAME(BeginVolumeToMeshTool, "Volume To Mesh", "Vol2Msh", "Convert a Volume to a new Mesh Object");
+	REGISTER_MODELING_TOOL_COMMAND_WITH_SHORTNAME(BeginMeshToVolumeTool, "Mesh To Volume", "Msh2Vol", "Convert a Mesh to a Volume");
+	if (!Settings->InRestrictiveMode())
 	{
-		REGISTER_MODELING_TOOL_COMMAND(BeginVolumeToMeshTool, "Volume To Mesh", "Convert a Volume to a new Mesh Object");
-		REGISTER_MODELING_TOOL_COMMAND(BeginMeshToVolumeTool, "Mesh To Volume", "Convert a Mesh to a Volume");
-		if (!Settings->InRestrictiveMode())
-		{
-			REGISTER_MODELING_TOOL_COMMAND(BeginBspConversionTool, "Convert BSPs", "Convert BSP to a new Mesh Object");
-		}
-		REGISTER_MODELING_TOOL_COMMAND(BeginPhysicsInspectorTool, "Inspect Collision", "Inspect the physics/collision geometry for selected Meshes");
-		REGISTER_MODELING_TOOL_COMMAND(BeginSimpleCollisionEditorTool, "Simple Collision Editor", "Edit the simple collision shapes for the selected Mesh");
-		REGISTER_MODELING_TOOL_COMMAND(BeginSetCollisionGeometryTool, "Mesh To Collision", "Convert selected Meshes to Simple Collision Geometry (for last selected)");
-		REGISTER_MODELING_TOOL_COMMAND(BeginExtractCollisionGeometryTool, "Collision To Mesh", "Convert Simple Collision Geometry to a new Mesh Object");
+		REGISTER_MODELING_TOOL_COMMAND_WITH_SHORTNAME(BeginBspConversionTool, "Convert BSPs", "BSPConv", "Convert BSP to a new Mesh Object");
 	}
+	REGISTER_MODELING_TOOL_COMMAND_WITH_SHORTNAME(BeginPhysicsInspectorTool, "Inspect Collision", "PInspct", "Inspect the physics/collision geometry for selected Meshes");
+	REGISTER_MODELING_TOOL_COMMAND_WITH_SHORTNAME(BeginSimpleCollisionEditorTool, "Simple Collision Editor", "SCollEdit", "Edit the simple collision shapes for the selected Mesh");
+	REGISTER_MODELING_TOOL_COMMAND_WITH_SHORTNAME(BeginSetCollisionGeometryTool, "Mesh To Collision", "Msh2Coll", "Convert selected Meshes to Simple Collision Geometry (for last selected)");
+	REGISTER_MODELING_TOOL_COMMAND_WITH_SHORTNAME(BeginExtractCollisionGeometryTool, "Collision To Mesh", "Coll2Msh", "Convert Simple Collision Geometry to a new Mesh Object");
 
 	// LODs
-	if ( bUseShortNames )
-	{
-		REGISTER_MODELING_TOOL_COMMAND(BeginLODManagerTool, "LODMgr", "Inspect the LODs of a Static Mesh Asset");
-		REGISTER_MODELING_TOOL_COMMAND(BeginGenerateStaticMeshLODAssetTool, "AutoLOD", "Automatically generate a simplified LOD with baked Textures/Materials for a Mesh");
-		REGISTER_MODELING_TOOL_COMMAND(BeginISMEditorTool, "ISMEd", "Edit the Instances of Instanced Static Mesh Components");
-	}
-	else
-	{
-		REGISTER_MODELING_TOOL_COMMAND(BeginLODManagerTool, "LOD Manager", "Inspect the LODs of a Static Mesh Asset");
-		REGISTER_MODELING_TOOL_COMMAND(BeginGenerateStaticMeshLODAssetTool, "AutoLOD", "Automatically generate a simplified LOD with baked Textures/Materials for a Mesh");
-		REGISTER_MODELING_TOOL_COMMAND(BeginISMEditorTool, "ISM Editor", "Edit the Instances of Instanced Static Mesh Components");
-	}
+	REGISTER_MODELING_TOOL_COMMAND_WITH_SHORTNAME(BeginLODManagerTool, "LOD Manager", "LODMgr", "Inspect the LODs of a Static Mesh Asset");
+	REGISTER_MODELING_TOOL_COMMAND(BeginGenerateStaticMeshLODAssetTool, "AutoLOD", "Automatically generate a simplified LOD with baked Textures/Materials for a Mesh");
+	REGISTER_MODELING_TOOL_COMMAND_WITH_SHORTNAME(BeginISMEditorTool, "ISM Editor", "ISMEd", "Edit the Instances of Instanced Static Mesh Components");
 
 	REGISTER_MODELING_TOOL_COMMAND(BeginAddPatchTool, "Patch", "Add Patch");
 	REGISTER_MODELING_TOOL_COMMAND(BeginShapeSprayTool, "Spray", "Shape Spray");
@@ -339,35 +214,13 @@ void FModelingToolsManagerCommands::RegisterCommands()
 	REGISTER_MODELING_TOOL_COMMAND(BeginPolyModelTool_Inset, "Inset", "Inset the current Selected Faces");
 	REGISTER_MODELING_TOOL_COMMAND(BeginPolyModelTool_Outset, "Outset", "Outset the current Selected Faces");
 	REGISTER_MODELING_TOOL_COMMAND(BeginPolyModelTool_CutFaces, "Cut", "Cut the current Selected Faces");
-	if (bUseShortNames)
-	{
-		REGISTER_MODELING_TOOL_COMMAND(BeginPolyModelTool_ExtrudeEdges, "ExtEdg", "Extrude selected boundary edges.");
-	}
-	else
-	{
-		REGISTER_MODELING_TOOL_COMMAND(BeginPolyModelTool_ExtrudeEdges, "Extrude Edges", "Extrude selected boundary edges.");
-	}
+	REGISTER_MODELING_TOOL_COMMAND_WITH_SHORTNAME(BeginPolyModelTool_ExtrudeEdges, "Extrude Edges", "ExtEdg", "Extrude selected boundary edges.");
 	REGISTER_MODELING_TOOL_COMMAND(BeginPolyModelTool_PushPull, "PushPull", "Push/Pull the current Selected Faces");
 	REGISTER_MODELING_TOOL_COMMAND(BeginPolyModelTool_Bevel, "Bevel", "Bevel the current Mesh Selection (Edges or Faces)");
-	if ( bUseShortNames )
-	{
-		REGISTER_MODELING_TOOL_COMMAND(BeginPolyModelTool_InsertEdgeLoop, "ELoop", "Insert Edge Loops into the Selected Mesh");
-	}
-	else
-	{
-		REGISTER_MODELING_TOOL_COMMAND(BeginPolyModelTool_InsertEdgeLoop, "Insert Loops", "Insert Edge Loops into the Selected Mesh");
-	}
+	REGISTER_MODELING_TOOL_COMMAND_WITH_SHORTNAME(BeginPolyModelTool_InsertEdgeLoop, "Insert Loops", "ELoop", "Insert Edge Loops into the Selected Mesh");
 
-	if ( bUseShortNames )
-	{
-		REGISTER_MODELING_TOOL_COMMAND(BeginPolyModelTool_PolyEd, "PolyEd", "Select / Edit the current Mesh via PolyGroups");
-		REGISTER_MODELING_TOOL_COMMAND(BeginPolyModelTool_TriSel, "TriSel", "Select / Edit the current Mesh via Triangles");
-	}
-	else
-	{
-		REGISTER_MODELING_TOOL_COMMAND(BeginPolyModelTool_PolyEd, "PolyGroup Edit", "Select / Edit the current Mesh via PolyGroups");
-		REGISTER_MODELING_TOOL_COMMAND(BeginPolyModelTool_TriSel, "Tri Select", "Select / Edit the current Mesh triangles with a brush interface");
-	}
+	REGISTER_MODELING_TOOL_COMMAND_WITH_SHORTNAME(BeginPolyModelTool_PolyEd, "PolyGroup Edit", "PolyEd", "Select / Edit the current Mesh via PolyGroups");
+	REGISTER_MODELING_TOOL_COMMAND_WITH_SHORTNAME(BeginPolyModelTool_TriSel, "Tri Select", "TriSel", "Select / Edit the current Mesh triangles with a brush interface");
 
 
 	REGISTER_MODELING_TOOL_COMMAND(MeshSelectionModeAction_NoSelection, "None", "Disable Mesh Element Selection (Select Objects Only)");
@@ -380,14 +233,7 @@ void FModelingToolsManagerCommands::RegisterCommands()
 
 	REGISTER_MODELING_ACTION_COMMAND(BeginSelectionAction_Delete, "Delete", "Delete the current Mesh Selection");
 	REGISTER_MODELING_ACTION_COMMAND(BeginSelectionAction_Retriangulate, "Clean", "Retriangulate the current Mesh or Mesh Selection");
-	if ( bUseShortNames )
-	{
-		REGISTER_MODELING_ACTION_COMMAND(BeginSelectionAction_Disconnect, "Discon", "Disconnect the current Mesh Selection");
-	}
-	else
-	{
-		REGISTER_MODELING_ACTION_COMMAND(BeginSelectionAction_Disconnect, "Disconnect", "Disconnect the current Mesh Selection");
-	}
+	REGISTER_MODELING_ACTION_COMMAND_WITH_SHORTNAME(BeginSelectionAction_Disconnect, "Disconnect", "Discon", "Disconnect the current Mesh Selection");
 
 	REGISTER_MODELING_TOOL_COMMAND(BeginSelectionAction_Extrude, "Extrude", "Extrude the current Mesh Selection");
 	REGISTER_MODELING_TOOL_COMMAND(BeginSelectionAction_Offset, "Offset", "Offset the current Mesh Selection");
@@ -429,7 +275,14 @@ void FModelingToolsManagerCommands::RegisterCommands()
 	UI_COMMAND(AcceptOrCompleteActiveTool, "Accept or Complete", "Accept or Complete the active Tool", EUserInterfaceActionType::Button, FInputChord(EKeys::Enter));
 	UI_COMMAND(CancelOrCompleteActiveTool, "Cancel or Complete", "Cancel or Complete the active Tool", EUserInterfaceActionType::Button, FInputChord(EKeys::Escape));
 	
+	for (int32 i = 0; i < RegisteredTools.Num(); ++i)
+	{
+		CommandToRegisteredToolsIndex.Add(RegisteredTools[i].ToolCommand.Get(), i);
+	}
 #undef REGISTER_MODELING_TOOL_COMMAND
+#undef REGISTER_MODELING_TOOL_COMMAND_WITH_SHORTNAME
+#undef REGISTER_MODELING_ACTION_COMMAND
+#undef REGISTER_MODELING_ACTION_COMMAND_WITH_SHORTNAME
 }
 
 
diff --git a/Engine/Plugins/Editor/ModelingToolsEditorMode/Source/ModelingToolsEditorMode/Public/ModelingToolsManagerActions.h b/Engine/Plugins/Editor/ModelingToolsEditorMode/Source/ModelingToolsEditorMode/Public/ModelingToolsManagerActions.h
--- a/Engine/Plugins/Editor/ModelingToolsEditorMode/Source/ModelingToolsEditorMode/Public/ModelingToolsManagerActions.h
+++ b/Engine/Plugins/Editor/ModelingToolsEditorMode/Source/ModelingToolsEditorMode/Public/ModelingToolsManagerActions.h
@@ -18,6 +18,9 @@ protected:
 	{
 		FString ToolUIName;
 		TSharedPtr<FUICommandInfo> ToolCommand;
+
+		// Optional (not used if empty). Only relevant if using short labels for tools
+		FText ShortName;
 	};
 	TArray<FStartToolCommand> RegisteredTools;		// Tool start-commands listed below are stored in this list
 
@@ -34,6 +37,12 @@ public:
 	 */
 	TSharedPtr<FUICommandInfo> FindToolByName(FString Name, bool& bFound) const;
 	
+	/**
+	 * Helper that gets the regular command label (result of GetLabel) or a shortened version
+	 *  if one was registered and bShortName is true.
+	 */
+	FText GetCommandLabel(TSharedPtr<const FUICommandInfo> Command, bool bShortName) const;
+
 	//
 	// These commands are set up to launch registered Tools via the ToolManager in ModelingToolsEditorMode.cpp
 	//
@@ -235,6 +244,9 @@ public:
 		const FText& Tooltip,
 		const FSlateIcon& Icon);
 
+private:
+	// Gets index into RegisteredTools for a given command.
+	TMap<FUICommandInfo*, int32> CommandToRegisteredToolsIndex;
 };
 
 
