MeshModelingTools: Fix the tool command LOCTEXT entries for short names not being parsed by the localization processor.

diff --git a/Engine/Plugins/Editor/ModelingToolsEditorMode/Source/ModelingToolsEditorMode/Private/ModelingToolsManagerActions.cpp b/Engine/Plugins/Editor/ModelingToolsEditorMode/Source/ModelingToolsEditorMode/Private/ModelingToolsManagerActions.cpp
--- a/Engine/Plugins/Editor/ModelingToolsEditorMode/Source/ModelingToolsEditorMode/Private/ModelingToolsManagerActions.cpp
+++ b/Engine/Plugins/Editor/ModelingToolsEditorMode/Source/ModelingToolsEditorMode/Private/ModelingToolsManagerActions.cpp
@@ -57,23 +57,30 @@ void FModelingToolsManagerCommands::RegisterCommands()
 {
 	const UModelingToolsEditorModeSettings* Settings = GetDefault<UModelingToolsEditorModeSettings>();
 
-	// this has to be done with a compile-time macro because UI_COMMAND expands to LOCTEXT macros
+	// Some points about registering commands:
+	// 1. UI_COMMAND expands to a LOCTEXT macro, so the helpers need to be compile-time macros.
+	// 2. The LOCTEXT expansion uses ToolCommandInfo as the key, so trying to conditionally register different
+	//  tool names (e.g. short names) for the same ToolCommandInfo will lead to LOCTEXT key collisions unless
+	//  you define a different LOCTEXT_NAMESPACE in the different branches. We could have done that for handling
+	//  short names, but decided instead to store short names ourselves.
+	// 3. LOCTEXT processor is a simple pattern matcher (though with special handling for UI_COMMAND) and
+	//  does not run after preprocessing, so trying to assemble a LOCTEXT in the macro, such as 
+	//  "LOCTEXT(ToolName"_short", ShortName)" will not work. Hence us passing a ready LOCTEXT
+	//  to our macros for the short name.
+
 #define REGISTER_MODELING_TOOL_COMMAND(ToolCommandInfo, ToolName, ToolTip) \
 		UI_COMMAND(ToolCommandInfo, ToolName, ToolTip, EUserInterfaceActionType::ToggleButton, FInputChord()); \
 		RegisteredTools.Add(FStartToolCommand{ ToolName, ToolCommandInfo });
-	// If using short names, we can't use the same command with a different ToolName, because that will
-	//  generate colliding LOCTEXT keys. Instead we keep the short name in a map.
-#define REGISTER_MODELING_TOOL_COMMAND_WITH_SHORTNAME(ToolCommandInfo, ToolName, ShortName, ToolTip) \
+#define REGISTER_MODELING_TOOL_COMMAND_WITH_SHORTNAME(ToolCommandInfo, ToolName, ShortNameFText, ToolTip) \
 		UI_COMMAND(ToolCommandInfo, ToolName, ToolTip, EUserInterfaceActionType::ToggleButton, FInputChord()); \
-		RegisteredTools.Add(FStartToolCommand{ ToolName, ToolCommandInfo, LOCTEXT(ToolName"_short", ShortName) });
+		RegisteredTools.Add(FStartToolCommand{ ToolName, ToolCommandInfo, ShortNameFText });
 
-	// this has to be done with a compile-time macro because UI_COMMAND expands to LOCTEXT macros
 #define REGISTER_MODELING_ACTION_COMMAND(ToolCommandInfo, ToolName, ToolTip) \
 		UI_COMMAND(ToolCommandInfo, ToolName, ToolTip, EUserInterfaceActionType::Button, FInputChord()); \
 		RegisteredTools.Add(FStartToolCommand{ ToolName, ToolCommandInfo });
-#define REGISTER_MODELING_ACTION_COMMAND_WITH_SHORTNAME(ToolCommandInfo, ToolName, ShortName, ToolTip) \
+#define REGISTER_MODELING_ACTION_COMMAND_WITH_SHORTNAME(ToolCommandInfo, ToolName, ShortNameFText, ToolTip) \
 		UI_COMMAND(ToolCommandInfo, ToolName, ToolTip, EUserInterfaceActionType::Button, FInputChord()); \
-		RegisteredTools.Add(FStartToolCommand{ ToolName, ToolCommandInfo, LOCTEXT(ToolName"_short", ShortName) });
+		RegisteredTools.Add(FStartToolCommand{ ToolName, ToolCommandInfo, ShortNameFText });
 
 	// Shapes
 	REGISTER_MODELING_TOOL_COMMAND(BeginAddBoxPrimitiveTool, "Box", "Create new box objects");
@@ -84,52 +91,52 @@ void FModelingToolsManagerCommands::RegisterCommands()
 	REGISTER_MODELING_TOOL_COMMAND(BeginAddDiscPrimitiveTool, "Disc", "Create new disc objects");
 	REGISTER_MODELING_TOOL_COMMAND(BeginAddStairsPrimitiveTool, "Stairs", "Create new stairs objects");
 	REGISTER_MODELING_TOOL_COMMAND(BeginAddCapsulePrimitiveTool, "Capsule", "Create new capsule objects");
-	REGISTER_MODELING_TOOL_COMMAND_WITH_SHORTNAME(BeginAddCylinderPrimitiveTool, "Cylinder", "Cyl", "Create new cylinder objects");
-	REGISTER_MODELING_TOOL_COMMAND_WITH_SHORTNAME(BeginAddRectanglePrimitiveTool, "Rectangle", "Rect", "Create new rectangle objects");
+	REGISTER_MODELING_TOOL_COMMAND_WITH_SHORTNAME(BeginAddCylinderPrimitiveTool, "Cylinder", LOCTEXT("BeginAddCylinderPrimitiveTool_short", "Cyl"), "Create new cylinder objects");
+	REGISTER_MODELING_TOOL_COMMAND_WITH_SHORTNAME(BeginAddRectanglePrimitiveTool, "Rectangle", LOCTEXT("BeginAddRectanglePrimitiveTool_short", "Rect"), "Create new rectangle objects");
 
 
 	// Create
-	REGISTER_MODELING_TOOL_COMMAND_WITH_SHORTNAME(BeginDrawPolygonTool, "Extrude Polygon", "PolyExt", "Draw and extrude 2D Polygons to create new objects");
-	REGISTER_MODELING_TOOL_COMMAND_WITH_SHORTNAME(BeginDrawPolyPathTool, "Extrude Path", "PathExt", "Draw and extrude 2D Paths to create new objects");
-	REGISTER_MODELING_TOOL_COMMAND_WITH_SHORTNAME(BeginDrawAndRevolveTool, "Revolve Path", "PathRev", "Draw and revolve 2D Paths to create new objects");
-	REGISTER_MODELING_TOOL_COMMAND_WITH_SHORTNAME(BeginRevolveSplineTool, "Revolve Spline", "SplnRev", "Revolve splines to create new objects");
-	REGISTER_MODELING_TOOL_COMMAND_WITH_SHORTNAME(BeginRevolveBoundaryTool, "Revolve Boundary", "BdryRev", "Revolve Mesh boundary loops to create new objects");
+	REGISTER_MODELING_TOOL_COMMAND_WITH_SHORTNAME(BeginDrawPolygonTool, "Extrude Polygon", LOCTEXT("BeginDrawPolygonTool_short", "PolyExt"), "Draw and extrude 2D Polygons to create new objects");
+	REGISTER_MODELING_TOOL_COMMAND_WITH_SHORTNAME(BeginDrawPolyPathTool, "Extrude Path", LOCTEXT("BeginDrawPolyPathTool_short", "PathExt"), "Draw and extrude 2D Paths to create new objects");
+	REGISTER_MODELING_TOOL_COMMAND_WITH_SHORTNAME(BeginDrawAndRevolveTool, "Revolve Path", LOCTEXT("BeginDrawAndRevolveTool_short", "PathRev"), "Draw and revolve 2D Paths to create new objects");
+	REGISTER_MODELING_TOOL_COMMAND_WITH_SHORTNAME(BeginRevolveSplineTool, "Revolve Spline", LOCTEXT("BeginRevolveSplineTool_short", "SplnRev"), "Revolve splines to create new objects");
+	REGISTER_MODELING_TOOL_COMMAND_WITH_SHORTNAME(BeginRevolveBoundaryTool, "Revolve Boundary", LOCTEXT("BeginRevolveBoundaryTool_short", "BdryRev"), "Revolve Mesh boundary loops to create new objects");
 	REGISTER_MODELING_TOOL_COMMAND(BeginCombineMeshesTool, "Merge", "Merge multiple Meshes to create new objects");
-	REGISTER_MODELING_TOOL_COMMAND_WITH_SHORTNAME(BeginDuplicateMeshesTool, "Duplicate", "Dupe", "Duplicate single Meshes to create new objects");
-	REGISTER_MODELING_TOOL_COMMAND_WITH_SHORTNAME(BeginDrawSplineTool, "Draw Spline", "DrwSpln", "Draw a spline in the level");
-	REGISTER_MODELING_TOOL_COMMAND_WITH_SHORTNAME(BeginTriangulateSplinesTool, "Mesh Splines", "MshSpln", "Triangulate the Spline Components of selected actors to create new objects");
+	REGISTER_MODELING_TOOL_COMMAND_WITH_SHORTNAME(BeginDuplicateMeshesTool, "Duplicate", LOCTEXT("BeginDuplicateMeshesTool_short", "Dupe"), "Duplicate single Meshes to create new objects");
+	REGISTER_MODELING_TOOL_COMMAND_WITH_SHORTNAME(BeginDrawSplineTool, "Draw Spline", LOCTEXT("BeginDrawSplineTool_short", "DrwSpln"), "Draw a spline in the level");
+	REGISTER_MODELING_TOOL_COMMAND_WITH_SHORTNAME(BeginTriangulateSplinesTool, "Mesh Splines", LOCTEXT("BeginTriangulateSplinesTool_short", "MshSpln"), "Triangulate the Spline Components of selected actors to create new objects");
 
 	// PolyModel
-	REGISTER_MODELING_TOOL_COMMAND_WITH_SHORTNAME(BeginPolyEditTool, "PolyGroup Edit", "PolyEd", "Edit Meshes via PolyGroups");
-	REGISTER_MODELING_TOOL_COMMAND_WITH_SHORTNAME(BeginPolyDeformTool, "Deform PolyGroups", "PolyDef", "Deform Meshes via PolyGroups");
-	REGISTER_MODELING_TOOL_COMMAND_WITH_SHORTNAME(BeginCubeGridTool, "CubeGrid", "CubeGr", "Create block out Meshes using a repositionable grid");
-	REGISTER_MODELING_TOOL_COMMAND_WITH_SHORTNAME(BeginMeshBooleanTool, "Boolean", "MshBool", "Apply Boolean operations to the two selected Meshes");
-	REGISTER_MODELING_TOOL_COMMAND_WITH_SHORTNAME(BeginCutMeshWithMeshTool, "Mesh Cut", "MshCut", "Split one Mesh into parts using a second Mesh");
-	REGISTER_MODELING_TOOL_COMMAND_WITH_SHORTNAME(BeginSubdividePolyTool, "Subdivide", "SubDiv", "Subdivide the selected Mesh via PolyGroups or triangles");
+	REGISTER_MODELING_TOOL_COMMAND_WITH_SHORTNAME(BeginPolyEditTool, "PolyGroup Edit", LOCTEXT("BeginPolyEditTool_short", "PolyEd"), "Edit Meshes via PolyGroups");
+	REGISTER_MODELING_TOOL_COMMAND_WITH_SHORTNAME(BeginPolyDeformTool, "Deform PolyGroups", LOCTEXT("BeginPolyDeformTool_short", "PolyDef"), "Deform Meshes via PolyGroups");
+	REGISTER_MODELING_TOOL_COMMAND_WITH_SHORTNAME(BeginCubeGridTool, "CubeGrid", LOCTEXT("BeginCubeGridTool_short", "CubeGr"), "Create block out Meshes using a repositionable grid");
+	REGISTER_MODELING_TOOL_COMMAND_WITH_SHORTNAME(BeginMeshBooleanTool, "Boolean", LOCTEXT("BeginMeshBooleanTool_short", "MshBool"), "Apply Boolean operations to the two selected Meshes");
+	REGISTER_MODELING_TOOL_COMMAND_WITH_SHORTNAME(BeginCutMeshWithMeshTool, "Mesh Cut", LOCTEXT("BeginCutMeshWithMeshTool_short", "MshCut"), "Split one Mesh into parts using a second Mesh");
+	REGISTER_MODELING_TOOL_COMMAND_WITH_SHORTNAME(BeginSubdividePolyTool, "Subdivide", LOCTEXT("BeginSubdividePolyTool_short", "SubDiv"), "Subdivide the selected Mesh via PolyGroups or triangles");
 
 	// TriModel
-	REGISTER_MODELING_TOOL_COMMAND_WITH_SHORTNAME(BeginMeshSelectionTool, "Tri Select", "TriSel", "Select and edit Mesh triangles with a brush interface");
-	REGISTER_MODELING_TOOL_COMMAND_WITH_SHORTNAME(BeginTriEditTool, "Triangle Edit", "TriEd", "Select and Edit the Mesh vertices, edges, and triangles");
-	REGISTER_MODELING_TOOL_COMMAND_WITH_SHORTNAME(BeginHoleFillTool, "Fill Holes", "HFill", "Fill holes in the selected Mesh");
-	REGISTER_MODELING_TOOL_COMMAND_WITH_SHORTNAME(BeginPlaneCutTool, "Plane Cut", "PlnCut", "Cut the selected Meshes with a 3D plane");
+	REGISTER_MODELING_TOOL_COMMAND_WITH_SHORTNAME(BeginMeshSelectionTool, "Tri Select", LOCTEXT("BeginMeshSelectionTool_short", "TriSel"), "Select and edit Mesh triangles with a brush interface");
+	REGISTER_MODELING_TOOL_COMMAND_WITH_SHORTNAME(BeginTriEditTool, "Triangle Edit", LOCTEXT("BeginTriEditTool_short", "TriEd"), "Select and Edit the Mesh vertices, edges, and triangles");
+	REGISTER_MODELING_TOOL_COMMAND_WITH_SHORTNAME(BeginHoleFillTool, "Fill Holes", LOCTEXT("BeginHoleFillTool_short", "HFill"), "Fill holes in the selected Mesh");
+	REGISTER_MODELING_TOOL_COMMAND_WITH_SHORTNAME(BeginPlaneCutTool, "Plane Cut", LOCTEXT("BeginPlaneCutTool_short", "PlnCut"), "Cut the selected Meshes with a 3D plane");
 	REGISTER_MODELING_TOOL_COMMAND(BeginMirrorTool, "Mirror", "Mirror the selected Meshes");
 	REGISTER_MODELING_TOOL_COMMAND(BeginPolygonCutTool, "PolyCut", "Cut the selected Mesh with an extruded polygon");
 	REGISTER_MODELING_TOOL_COMMAND(BeginMeshTrimTool, "Trim", "Trim/Cut the selected Mesh with the second selected Mesh");
 
 	// Deform
-	REGISTER_MODELING_TOOL_COMMAND_WITH_SHORTNAME(BeginSculptMeshTool, "Vertex Sculpt", "VSclpt", "Vertex sculpting");
-	REGISTER_MODELING_TOOL_COMMAND_WITH_SHORTNAME(BeginRemeshSculptMeshTool, "Dynamic Sculpt", "DSclpt", "Dynamic mesh sculpting");
-	REGISTER_MODELING_TOOL_COMMAND_WITH_SHORTNAME(BeginDisplaceMeshTool, "Displace", "Displce", "Tessellate and Displace the selected Mesh");
+	REGISTER_MODELING_TOOL_COMMAND_WITH_SHORTNAME(BeginSculptMeshTool, "Vertex Sculpt", LOCTEXT("BeginSculptMeshTool_short", "VSclpt"), "Vertex sculpting");
+	REGISTER_MODELING_TOOL_COMMAND_WITH_SHORTNAME(BeginRemeshSculptMeshTool, "Dynamic Sculpt", LOCTEXT("BeginRemeshSculptMeshTool_short", "DSclpt"), "Dynamic mesh sculpting");
+	REGISTER_MODELING_TOOL_COMMAND_WITH_SHORTNAME(BeginDisplaceMeshTool, "Displace", LOCTEXT("BeginDisplaceMeshTool_short", "Displce"), "Tessellate and Displace the selected Mesh");
 	REGISTER_MODELING_TOOL_COMMAND(BeginSmoothMeshTool, "Smooth", "Smooth the shape of the selected Mesh");
 	REGISTER_MODELING_TOOL_COMMAND(BeginOffsetMeshTool, "Offset", "Offset the surface of the selected Mesh");
 	REGISTER_MODELING_TOOL_COMMAND(BeginMeshSpaceDeformerTool, "Warp", "Reshape the selected Mesh using space deformers");
 	REGISTER_MODELING_TOOL_COMMAND(BeginLatticeDeformerTool, "Lattice", "Deform the selected Mesh using a 3D lattice/grid");
 
 	// Transform
-	REGISTER_MODELING_TOOL_COMMAND_WITH_SHORTNAME(BeginTransformMeshesTool, "Transform", "XForm", "Transform the selected Meshes");
-	REGISTER_MODELING_TOOL_COMMAND_WITH_SHORTNAME(BeginEditPivotTool, "Edit Pivot", "Pivot", "Edit the pivot points of the selected Meshes");
-	REGISTER_MODELING_TOOL_COMMAND_WITH_SHORTNAME(BeginAddPivotActorTool, "Pivot Actor", "PivotAct", "Add actor to act as a pivot for child components");
-	REGISTER_MODELING_TOOL_COMMAND_WITH_SHORTNAME(BeginBakeTransformTool, "Bake Transform", "BakeRS", "Bake rotation and scale into the selected Meshes");
+	REGISTER_MODELING_TOOL_COMMAND_WITH_SHORTNAME(BeginTransformMeshesTool, "Transform", LOCTEXT("BeginTransformMeshesTool_short", "XForm"), "Transform the selected Meshes");
+	REGISTER_MODELING_TOOL_COMMAND_WITH_SHORTNAME(BeginEditPivotTool, "Edit Pivot", LOCTEXT("BeginEditPivotTool_short", "Pivot"), "Edit the pivot points of the selected Meshes");
+	REGISTER_MODELING_TOOL_COMMAND_WITH_SHORTNAME(BeginAddPivotActorTool, "Pivot Actor", LOCTEXT("BeginAddPivotActorTool_short", "PivotAct"), "Add actor to act as a pivot for child components");
+	REGISTER_MODELING_TOOL_COMMAND_WITH_SHORTNAME(BeginBakeTransformTool, "Bake Transform", LOCTEXT("BeginBakeTransformTool_short", "BakeRS"), "Bake rotation and scale into the selected Meshes");
 	REGISTER_MODELING_TOOL_COMMAND(BeginAlignObjectsTool, "Align", "Align the selected Objects");
 	REGISTER_MODELING_TOOL_COMMAND(BeginTransferMeshTool, "Transfer", "Copy the first selected Mesh to the second selected Mesh");
 	REGISTER_MODELING_TOOL_COMMAND(BeginConvertMeshesTool, "Convert", "Convert the selected Meshes to a different type of Mesh Object");
@@ -139,7 +146,7 @@ void FModelingToolsManagerCommands::RegisterCommands()
 	
 
 	// MeshOps
-	REGISTER_MODELING_TOOL_COMMAND_WITH_SHORTNAME(BeginSimplifyMeshTool, "Simplify", "Simplfy", "Simplify the selected Mesh");
+	REGISTER_MODELING_TOOL_COMMAND_WITH_SHORTNAME(BeginSimplifyMeshTool, "Simplify", LOCTEXT("BeginSimplifyMeshTool_short", "Simplfy"), "Simplify the selected Mesh");
 	REGISTER_MODELING_TOOL_COMMAND(BeginRemeshMeshTool, "Remesh", "Re-triangulate the selected Mesh");
 	REGISTER_MODELING_TOOL_COMMAND(BeginWeldEdgesTool, "Weld", "Weld overlapping Mesh edges");
 	REGISTER_MODELING_TOOL_COMMAND(BeginRemoveOccludedTrianglesTool, "Jacket", "Remove hidden triangles from the selected Meshes");
@@ -148,59 +155,59 @@ void FModelingToolsManagerCommands::RegisterCommands()
 
 
 	// VoxOps
-	REGISTER_MODELING_TOOL_COMMAND_WITH_SHORTNAME(BeginVoxelSolidifyTool, "Voxel Wrap", "VoxWrap", "Wrap the selected Meshes using SDFs/voxels");
-	REGISTER_MODELING_TOOL_COMMAND_WITH_SHORTNAME(BeginVoxelBlendTool, "Voxel Blend", "VoxBlnd", "Blend the selected Meshes using SDFs/voxels");
-	REGISTER_MODELING_TOOL_COMMAND_WITH_SHORTNAME(BeginVoxelMorphologyTool, "Voxel Offset", "VoxMrph", "Offset/Inset the selected Meshes using SDFs/voxels");
+	REGISTER_MODELING_TOOL_COMMAND_WITH_SHORTNAME(BeginVoxelSolidifyTool, "Voxel Wrap", LOCTEXT("BeginVoxelSolidifyTool_short", "VoxWrap"), "Wrap the selected Meshes using SDFs/voxels");
+	REGISTER_MODELING_TOOL_COMMAND_WITH_SHORTNAME(BeginVoxelBlendTool, "Voxel Blend", LOCTEXT("BeginVoxelBlendTool_short", "VoxBlnd"), "Blend the selected Meshes using SDFs/voxels");
+	REGISTER_MODELING_TOOL_COMMAND_WITH_SHORTNAME(BeginVoxelMorphologyTool, "Voxel Offset", LOCTEXT("BeginVoxelMorphologyTool_short", "VoxMrph"), "Offset/Inset the selected Meshes using SDFs/voxels");
 #if WITH_PROXYLOD
 	// The ProxyLOD plugin is currently only available on Windows. Without it, the following tools do not work as expected.
-	REGISTER_MODELING_TOOL_COMMAND_WITH_SHORTNAME(BeginVoxelBooleanTool, "Voxel Boolean", "VoxBool", "Boolean the selected Meshes using SDFs/voxels");
-	REGISTER_MODELING_TOOL_COMMAND_WITH_SHORTNAME(BeginVoxelMergeTool, "Voxel Merge", "VoxMrg", "Merge the selected Meshes using SDFs/voxels");
+	REGISTER_MODELING_TOOL_COMMAND_WITH_SHORTNAME(BeginVoxelBooleanTool, "Voxel Boolean", LOCTEXT("BeginVoxelBooleanTool_short", "VoxBool"), "Boolean the selected Meshes using SDFs/voxels");
+	REGISTER_MODELING_TOOL_COMMAND_WITH_SHORTNAME(BeginVoxelMergeTool, "Voxel Merge", LOCTEXT("BeginVoxelMergeTool_short", "VoxMrg"), "Merge the selected Meshes using SDFs/voxels");
 #endif	// WITH_PROXYLOD
 
 	// Attributes
-	REGISTER_MODELING_TOOL_COMMAND_WITH_SHORTNAME(BeginMeshInspectorTool, "Inspect", "Inspct", "Inspect Mesh attributes");
-	REGISTER_MODELING_TOOL_COMMAND_WITH_SHORTNAME(BeginEditNormalsTool, "Normals", "Nrmls", "Recompute or Repair Normals");
-	REGISTER_MODELING_TOOL_COMMAND_WITH_SHORTNAME(BeginEditTangentsTool, "Tangents", "Tngnts", "Recompute Tangents");
-	REGISTER_MODELING_TOOL_COMMAND_WITH_SHORTNAME(BeginPolyGroupsTool, "Generate PolyGroups", "GrpGen", "Generate new PolyGroups");
-	REGISTER_MODELING_TOOL_COMMAND_WITH_SHORTNAME(BeginMeshGroupPaintTool, "Paint PolyGroups", "GrpPnt", "Paint new PolyGroups");
-	REGISTER_MODELING_TOOL_COMMAND_WITH_SHORTNAME(BeginMeshVertexPaintTool, "Paint Vertex Colors", "VtxPnt", "Paint Mesh Vertex Colors");
-	REGISTER_MODELING_TOOL_COMMAND_WITH_SHORTNAME(BeginAttributeEditorTool, "Edit Attributes", "AttrEd", "Edit/configure Mesh attributes");
-	REGISTER_MODELING_TOOL_COMMAND_WITH_SHORTNAME(BeginMeshAttributePaintTool, "Paint Maps", "MapPnt", "Paint attribute maps");
-	REGISTER_MODELING_TOOL_COMMAND_WITH_SHORTNAME(BeginEditMeshMaterialsTool, "Edit Materials", "MatEd", "Assign materials to selected triangles");
+	REGISTER_MODELING_TOOL_COMMAND_WITH_SHORTNAME(BeginMeshInspectorTool, "Inspect", LOCTEXT("BeginMeshInspectorTool_short", "Inspct"), "Inspect Mesh attributes");
+	REGISTER_MODELING_TOOL_COMMAND_WITH_SHORTNAME(BeginEditNormalsTool, "Normals", LOCTEXT("BeginEditNormalsTool_short", "Nrmls"), "Recompute or Repair Normals");
+	REGISTER_MODELING_TOOL_COMMAND_WITH_SHORTNAME(BeginEditTangentsTool, "Tangents", LOCTEXT("BeginEditTangentsTool_short", "Tngnts"), "Recompute Tangents");
+	REGISTER_MODELING_TOOL_COMMAND_WITH_SHORTNAME(BeginPolyGroupsTool, "Generate PolyGroups", LOCTEXT("BeginPolyGroupsTool_short", "GrpGen"), "Generate new PolyGroups");
+	REGISTER_MODELING_TOOL_COMMAND_WITH_SHORTNAME(BeginMeshGroupPaintTool, "Paint PolyGroups", LOCTEXT("BeginMeshGroupPaintTool_short", "GrpPnt"), "Paint new PolyGroups");
+	REGISTER_MODELING_TOOL_COMMAND_WITH_SHORTNAME(BeginMeshVertexPaintTool, "Paint Vertex Colors", LOCTEXT("BeginMeshVertexPaintTool_short", "VtxPnt"), "Paint Mesh Vertex Colors");
+	REGISTER_MODELING_TOOL_COMMAND_WITH_SHORTNAME(BeginAttributeEditorTool, "Edit Attributes", LOCTEXT("BeginAttributeEditorTool_short", "AttrEd"), "Edit/configure Mesh attributes");
+	REGISTER_MODELING_TOOL_COMMAND_WITH_SHORTNAME(BeginMeshAttributePaintTool, "Paint Maps", LOCTEXT("BeginMeshAttributePaintTool_short", "MapPnt"), "Paint attribute maps");
+	REGISTER_MODELING_TOOL_COMMAND_WITH_SHORTNAME(BeginEditMeshMaterialsTool, "Edit Materials", LOCTEXT("BeginEditMeshMaterialsTool_short", "MatEd"), "Assign materials to selected triangles");
 
 	// UVs
 	REGISTER_MODELING_TOOL_COMMAND(BeginGlobalUVGenerateTool, "AutoUV", "Automatically unwrap and pack UVs");
 	// This is done directly, not with the REGISTER_ macro, since we don't want it added to the tool list or use a toggle button
 	UI_COMMAND(LaunchUVEditor, "UVEditor", "Launch UV asset editor", EUserInterfaceActionType::Button, FInputChord());
-	REGISTER_MODELING_TOOL_COMMAND_WITH_SHORTNAME(BeginGroupUVGenerateTool, "UV Unwrap", "Unwrap", "Recompute UVs for existing UV islands or PolyGroups");
-	REGISTER_MODELING_TOOL_COMMAND_WITH_SHORTNAME(BeginUVProjectionTool, "Project UVs", "Project", "Compute UVs via projecting to simple shapes");
-	REGISTER_MODELING_TOOL_COMMAND_WITH_SHORTNAME(BeginUVSeamEditTool, "Edit UV Seams", "SeamEd", "Add UV seams");
-	REGISTER_MODELING_TOOL_COMMAND_WITH_SHORTNAME(BeginTransformUVIslandsTool, "Transform UVs", "XFormUV", "Transform UV islands in UV space");
-	REGISTER_MODELING_TOOL_COMMAND_WITH_SHORTNAME(BeginUVLayoutTool, "Layout UVs", "Layout", "Transform and Repack existing UVs");
-	REGISTER_MODELING_TOOL_COMMAND_WITH_SHORTNAME(BeginUVTransferTool, "Transfer UVs", "Transfer", "Transfer UVs from a simplified mesh to a higher resolution mesh");
+	REGISTER_MODELING_TOOL_COMMAND_WITH_SHORTNAME(BeginGroupUVGenerateTool, "UV Unwrap", LOCTEXT("BeginGroupUVGenerateTool_short", "Unwrap"), "Recompute UVs for existing UV islands or PolyGroups");
+	REGISTER_MODELING_TOOL_COMMAND_WITH_SHORTNAME(BeginUVProjectionTool, "Project UVs", LOCTEXT("BeginUVProjectionTool_short", "Project"), "Compute UVs via projecting to simple shapes");
+	REGISTER_MODELING_TOOL_COMMAND_WITH_SHORTNAME(BeginUVSeamEditTool, "Edit UV Seams", LOCTEXT("BeginUVSeamEditTool_short", "SeamEd"), "Add UV seams");
+	REGISTER_MODELING_TOOL_COMMAND_WITH_SHORTNAME(BeginTransformUVIslandsTool, "Transform UVs", LOCTEXT("BeginTransformUVIslandsTool_short", "XFormUV"), "Transform UV islands in UV space");
+	REGISTER_MODELING_TOOL_COMMAND_WITH_SHORTNAME(BeginUVLayoutTool, "Layout UVs", LOCTEXT("BeginUVLayoutTool_short", "Layout"), "Transform and Repack existing UVs");
+	REGISTER_MODELING_TOOL_COMMAND_WITH_SHORTNAME(BeginUVTransferTool, "Transfer UVs", LOCTEXT("BeginUVTransferTool_short", "Transfer"), "Transfer UVs from a simplified mesh to a higher resolution mesh");
 
 	// Baking
-	REGISTER_MODELING_TOOL_COMMAND_WITH_SHORTNAME(BeginBakeMeshAttributeMapsTool, "Bake Textures", "BakeTx", "Bake textures for a target Mesh");
-	REGISTER_MODELING_TOOL_COMMAND_WITH_SHORTNAME(BeginBakeMultiMeshAttributeMapsTool, "Bake All", "BakeAll", "Bake textures for a target Mesh from multiple source Meshes");
-	REGISTER_MODELING_TOOL_COMMAND_WITH_SHORTNAME(BeginBakeMeshAttributeVertexTool, "Bake Vertex Colors", "BakeVtx", "Bake vertex colors for a target Mesh");
-	REGISTER_MODELING_TOOL_COMMAND_WITH_SHORTNAME(BeginBakeRenderCaptureTool, "Bake RC", "BakeRC", "Bake renders into new textures for a target Mesh from multiple source Meshes");
+	REGISTER_MODELING_TOOL_COMMAND_WITH_SHORTNAME(BeginBakeMeshAttributeMapsTool, "Bake Textures", LOCTEXT("BeginBakeMeshAttributeMapsTool_short", "BakeTx"), "Bake textures for a target Mesh");
+	REGISTER_MODELING_TOOL_COMMAND_WITH_SHORTNAME(BeginBakeMultiMeshAttributeMapsTool, "Bake All", LOCTEXT("BeginBakeMultiMeshAttributeMapsTool_short", "BakeAll"), "Bake textures for a target Mesh from multiple source Meshes");
+	REGISTER_MODELING_TOOL_COMMAND_WITH_SHORTNAME(BeginBakeMeshAttributeVertexTool, "Bake Vertex Colors", LOCTEXT("BeginBakeMeshAttributeVertexTool_short", "BakeVtx"), "Bake vertex colors for a target Mesh");
+	REGISTER_MODELING_TOOL_COMMAND_WITH_SHORTNAME(BeginBakeRenderCaptureTool, "Bake RC", LOCTEXT("BeginBakeRenderCaptureTool_short", "BakeRC"), "Bake renders into new textures for a target Mesh from multiple source Meshes");
 
 	// Volumes
-	REGISTER_MODELING_TOOL_COMMAND_WITH_SHORTNAME(BeginVolumeToMeshTool, "Volume To Mesh", "Vol2Msh", "Convert a Volume to a new Mesh Object");
-	REGISTER_MODELING_TOOL_COMMAND_WITH_SHORTNAME(BeginMeshToVolumeTool, "Mesh To Volume", "Msh2Vol", "Convert a Mesh to a Volume");
+	REGISTER_MODELING_TOOL_COMMAND_WITH_SHORTNAME(BeginVolumeToMeshTool, "Volume To Mesh", LOCTEXT("BeginVolumeToMeshTool_short", "Vol2Msh"), "Convert a Volume to a new Mesh Object");
+	REGISTER_MODELING_TOOL_COMMAND_WITH_SHORTNAME(BeginMeshToVolumeTool, "Mesh To Volume", LOCTEXT("BeginMeshToVolumeTool_short", "Msh2Vol"), "Convert a Mesh to a Volume");
 	if (!Settings->InRestrictiveMode())
 	{
-		REGISTER_MODELING_TOOL_COMMAND_WITH_SHORTNAME(BeginBspConversionTool, "Convert BSPs", "BSPConv", "Convert BSP to a new Mesh Object");
+		REGISTER_MODELING_TOOL_COMMAND_WITH_SHORTNAME(BeginBspConversionTool, "Convert BSPs", LOCTEXT("BeginBspConversionTool_short", "BSPConv"), "Convert BSP to a new Mesh Object");
 	}
-	REGISTER_MODELING_TOOL_COMMAND_WITH_SHORTNAME(BeginPhysicsInspectorTool, "Inspect Collision", "PInspct", "Inspect the physics/collision geometry for selected Meshes");
-	REGISTER_MODELING_TOOL_COMMAND_WITH_SHORTNAME(BeginSimpleCollisionEditorTool, "Simple Collision Editor", "SCollEdit", "Edit the simple collision shapes for the selected Mesh");
-	REGISTER_MODELING_TOOL_COMMAND_WITH_SHORTNAME(BeginSetCollisionGeometryTool, "Mesh To Collision", "Msh2Coll", "Convert selected Meshes to Simple Collision Geometry (for last selected)");
-	REGISTER_MODELING_TOOL_COMMAND_WITH_SHORTNAME(BeginExtractCollisionGeometryTool, "Collision To Mesh", "Coll2Msh", "Convert Simple Collision Geometry to a new Mesh Object");
+	REGISTER_MODELING_TOOL_COMMAND_WITH_SHORTNAME(BeginPhysicsInspectorTool, "Inspect Collision", LOCTEXT("BeginPhysicsInspectorTool_short", "PInspct"), "Inspect the physics/collision geometry for selected Meshes");
+	REGISTER_MODELING_TOOL_COMMAND_WITH_SHORTNAME(BeginSimpleCollisionEditorTool, "Simple Collision Editor", LOCTEXT("BeginSimpleCollisionEditorTool_short", "SCollEdit"), "Edit the simple collision shapes for the selected Mesh");
+	REGISTER_MODELING_TOOL_COMMAND_WITH_SHORTNAME(BeginSetCollisionGeometryTool, "Mesh To Collision", LOCTEXT("BeginSetCollisionGeometryTool_short", "Msh2Coll"), "Convert selected Meshes to Simple Collision Geometry (for last selected)");
+	REGISTER_MODELING_TOOL_COMMAND_WITH_SHORTNAME(BeginExtractCollisionGeometryTool, "Collision To Mesh", LOCTEXT("BeginExtractCollisionGeometryTool_short", "Coll2Msh"), "Convert Simple Collision Geometry to a new Mesh Object");
 
 	// LODs
-	REGISTER_MODELING_TOOL_COMMAND_WITH_SHORTNAME(BeginLODManagerTool, "LOD Manager", "LODMgr", "Inspect the LODs of a Static Mesh Asset");
+	REGISTER_MODELING_TOOL_COMMAND_WITH_SHORTNAME(BeginLODManagerTool, "LOD Manager", LOCTEXT("BeginLODManagerTool_short", "LODMgr"), "Inspect the LODs of a Static Mesh Asset");
 	REGISTER_MODELING_TOOL_COMMAND(BeginGenerateStaticMeshLODAssetTool, "AutoLOD", "Automatically generate a simplified LOD with baked Textures/Materials for a Mesh");
-	REGISTER_MODELING_TOOL_COMMAND_WITH_SHORTNAME(BeginISMEditorTool, "ISM Editor", "ISMEd", "Edit the Instances of Instanced Static Mesh Components");
+	REGISTER_MODELING_TOOL_COMMAND_WITH_SHORTNAME(BeginISMEditorTool, "ISM Editor", LOCTEXT("BeginISMEditorTool_short", "ISMEd"), "Edit the Instances of Instanced Static Mesh Components");
 
 	REGISTER_MODELING_TOOL_COMMAND(BeginAddPatchTool, "Patch", "Add Patch");
 	REGISTER_MODELING_TOOL_COMMAND(BeginShapeSprayTool, "Spray", "Shape Spray");
@@ -214,13 +221,13 @@ void FModelingToolsManagerCommands::RegisterCommands()
 	REGISTER_MODELING_TOOL_COMMAND(BeginPolyModelTool_Inset, "Inset", "Inset the current Selected Faces");
 	REGISTER_MODELING_TOOL_COMMAND(BeginPolyModelTool_Outset, "Outset", "Outset the current Selected Faces");
 	REGISTER_MODELING_TOOL_COMMAND(BeginPolyModelTool_CutFaces, "Cut", "Cut the current Selected Faces");
-	REGISTER_MODELING_TOOL_COMMAND_WITH_SHORTNAME(BeginPolyModelTool_ExtrudeEdges, "Extrude Edges", "ExtEdg", "Extrude selected boundary edges.");
+	REGISTER_MODELING_TOOL_COMMAND_WITH_SHORTNAME(BeginPolyModelTool_ExtrudeEdges, "Extrude Edges", LOCTEXT("BeginPolyModelTool_ExtrudeEdges_short", "ExtEdg"), "Extrude selected boundary edges.");
 	REGISTER_MODELING_TOOL_COMMAND(BeginPolyModelTool_PushPull, "PushPull", "Push/Pull the current Selected Faces");
 	REGISTER_MODELING_TOOL_COMMAND(BeginPolyModelTool_Bevel, "Bevel", "Bevel the current Mesh Selection (Edges or Faces)");
-	REGISTER_MODELING_TOOL_COMMAND_WITH_SHORTNAME(BeginPolyModelTool_InsertEdgeLoop, "Insert Loops", "ELoop", "Insert Edge Loops into the Selected Mesh");
+	REGISTER_MODELING_TOOL_COMMAND_WITH_SHORTNAME(BeginPolyModelTool_InsertEdgeLoop, "Insert Loops", LOCTEXT("BeginPolyModelTool_InsertEdgeLoop_short", "ELoop"), "Insert Edge Loops into the Selected Mesh");
 
-	REGISTER_MODELING_TOOL_COMMAND_WITH_SHORTNAME(BeginPolyModelTool_PolyEd, "PolyGroup Edit", "PolyEd", "Select / Edit the current Mesh via PolyGroups");
-	REGISTER_MODELING_TOOL_COMMAND_WITH_SHORTNAME(BeginPolyModelTool_TriSel, "Tri Select", "TriSel", "Select / Edit the current Mesh triangles with a brush interface");
+	REGISTER_MODELING_TOOL_COMMAND_WITH_SHORTNAME(BeginPolyModelTool_PolyEd, "PolyGroup Edit", LOCTEXT("BeginPolyModelTool_PolyEd_short", "PolyEd"), "Select / Edit the current Mesh via PolyGroups");
+	REGISTER_MODELING_TOOL_COMMAND_WITH_SHORTNAME(BeginPolyModelTool_TriSel, "Tri Select", LOCTEXT("BeginPolyModelTool_TriSel_short", "TriSel"), "Select / Edit the current Mesh triangles with a brush interface");
 
 
 	REGISTER_MODELING_TOOL_COMMAND(MeshSelectionModeAction_NoSelection, "None", "Disable Mesh Element Selection (Select Objects Only)");
@@ -233,7 +240,7 @@ void FModelingToolsManagerCommands::RegisterCommands()
 
 	REGISTER_MODELING_ACTION_COMMAND(BeginSelectionAction_Delete, "Delete", "Delete the current Mesh Selection");
 	REGISTER_MODELING_ACTION_COMMAND(BeginSelectionAction_Retriangulate, "Clean", "Retriangulate the current Mesh or Mesh Selection");
-	REGISTER_MODELING_ACTION_COMMAND_WITH_SHORTNAME(BeginSelectionAction_Disconnect, "Disconnect", "Discon", "Disconnect the current Mesh Selection");
+	REGISTER_MODELING_ACTION_COMMAND_WITH_SHORTNAME(BeginSelectionAction_Disconnect, "Disconnect", LOCTEXT("BeginSelectionAction_Disconnect_short", "Discon"), "Disconnect the current Mesh Selection");
 
 	REGISTER_MODELING_TOOL_COMMAND(BeginSelectionAction_Extrude, "Extrude", "Extrude the current Mesh Selection");
 	REGISTER_MODELING_TOOL_COMMAND(BeginSelectionAction_Offset, "Offset", "Offset the current Mesh Selection");
