warn in modeling mode when selection includes engine assets which cannot be modified

diff --git a/Engine/Plugins/Editor/ModelingToolsEditorMode/Source/ModelingToolsEditorMode/Private/ModelingToolsEditorModeToolkit.cpp b/Engine/Plugins/Editor/ModelingToolsEditorMode/Source/ModelingToolsEditorMode/Private/ModelingToolsEditorModeToolkit.cpp
--- a/Engine/Plugins/Editor/ModelingToolsEditorMode/Source/ModelingToolsEditorMode/Private/ModelingToolsEditorModeToolkit.cpp
+++ b/Engine/Plugins/Editor/ModelingToolsEditorMode/Source/ModelingToolsEditorMode/Private/ModelingToolsEditorModeToolkit.cpp
@@ -68,6 +68,9 @@
 #include "Widgets/Input/SMultiLineEditableTextBox.h"
 #include "IAssetViewport.h"
 
+// for editor selection queries
+#include "Selection.h"
+
 #if ENABLE_STYLUS_SUPPORT
 #include "ModelingStylusInputHandler.h"
 #endif
@@ -2296,6 +2299,10 @@ void FModelingToolsEditorModeToolkit::ShowRealtimeAndModeWarnings(bool bShowReal
 	{
 		WarningText = LOCTEXT("ModelingModeToolkitRealtimeWarning", "Realtime Mode is required for Modeling Tools to work correctly. Please enable Realtime Mode in the Viewport Options or with the Ctrl+r hotkey.");
 	}
+	else if (EngineAssetsSelected())
+	{
+		WarningText = LOCTEXT("ModelingModeToolkitEngineAssetsWarning", "Selection includes Engine assets, which cannot be modified.");
+	}
 	if (!WarningText.IdenticalTo(ActiveWarning))
 	{
 		ActiveWarning = WarningText;
@@ -2724,4 +2731,31 @@ IToolStylusStateProviderAPI* FModelingToolsEditorModeToolkit::GetStylusStateProv
 }
 
 
+bool FModelingToolsEditorModeToolkit::EngineAssetsSelected() const
+{
+	auto HasEngineAsset = [](UStaticMeshComponent* SMC) -> bool
+	{
+		UStaticMesh* StaticMesh = SMC->GetStaticMesh();
+		return StaticMesh && StaticMesh->GetPathName().StartsWith(TEXT("/Engine/"));
+	};
+	bool bFoundEngineAssets = false;
+	for (FSelectionIterator It(GEditor->GetSelectedActorIterator()); It; ++It)
+	{
+		AActor* Actor = static_cast<AActor*>(*It);
+		Actor->ForEachComponent<UStaticMeshComponent>(false, [&HasEngineAsset, &bFoundEngineAssets](UStaticMeshComponent* SMC)
+		{
+			bFoundEngineAssets = bFoundEngineAssets || HasEngineAsset(SMC);
+		});
+	}
+	for (FSelectionIterator It(GEditor->GetSelectedComponentIterator()); It; ++It)
+	{
+		if (UStaticMeshComponent* SMC = Cast<UStaticMeshComponent>(*It))
+		{
+			bFoundEngineAssets = bFoundEngineAssets || HasEngineAsset(SMC);
+		}
+	}
+	return bFoundEngineAssets;
+}
+
+
 #undef LOCTEXT_NAMESPACE
diff --git a/Engine/Plugins/Editor/ModelingToolsEditorMode/Source/ModelingToolsEditorMode/Public/ModelingToolsEditorModeToolkit.h b/Engine/Plugins/Editor/ModelingToolsEditorMode/Source/ModelingToolsEditorMode/Public/ModelingToolsEditorModeToolkit.h
--- a/Engine/Plugins/Editor/ModelingToolsEditorMode/Source/ModelingToolsEditorMode/Public/ModelingToolsEditorModeToolkit.h
+++ b/Engine/Plugins/Editor/ModelingToolsEditorMode/Source/ModelingToolsEditorMode/Public/ModelingToolsEditorModeToolkit.h
@@ -231,6 +231,9 @@ private:
 	TUniquePtr<UE::Modeling::FStylusInputHandler> StylusInputHandler;
 #endif
 
+	// tests if selection includes components backed by assets from the /Engine folder, of types the tools might otherwise modify (currently, static meshes)
+	bool EngineAssetsSelected() const;
+
 };
 
 #undef LOCTEXT_NAMESPACE
