Modeling Tools: Updates Mesh Element Selection default colors

diff --git a/Engine/Plugins/Editor/ModelingToolsEditorMode/Source/ModelingToolsEditorMode/Public/ModelingToolsEditorModeSettings.h b/Engine/Plugins/Editor/ModelingToolsEditorMode/Source/ModelingToolsEditorMode/Public/ModelingToolsEditorModeSettings.h
--- a/Engine/Plugins/Editor/ModelingToolsEditorMode/Source/ModelingToolsEditorMode/Public/ModelingToolsEditorModeSettings.h
+++ b/Engine/Plugins/Editor/ModelingToolsEditorMode/Source/ModelingToolsEditorMode/Public/ModelingToolsEditorModeSettings.h
@@ -355,25 +355,25 @@ public:
 	 * Custom color for when geometry is not a part of current selection when using Mesh Element Selection
 	 */
 	UPROPERTY(Config, EditAnywhere, Category = "Modeling Mode|Mesh Element Selection Mode")
-	FLinearColor UnselectedColor = FLinearColor(.906f, 0.038f, 0.038f,.4f); // default red
+	FLinearColor UnselectedColor = FLinearColor(.068478f, 0.56478f, 0.708376f,.5f); // default red
 
 	/**
 	 * Custom color hovering over selected geometry in Mesh Element Selection
 	 */
 	UPROPERTY(Config, EditAnywhere, Category = "Modeling Mode|Mesh Element Selection Mode")
-	FLinearColor HoverOverSelectedColor = FLinearColor(0.791f, 0.262f, 1.0f); // default pink
+	FLinearColor HoverOverSelectedColor = FLinearColor(1.0f, 0.943014f, 0.74f); // default pink
 
 	/**
 	* Custom color hovering over unselected geometry in Mesh Element Selection
 	*/
 	UPROPERTY(Config, EditAnywhere, Category = "Modeling Mode|Mesh Element Selection Mode")
-	FLinearColor HoverOverUnselectedColor = FLinearColor(0.027f, 0.656f, 0.027f, .6f); // default green
+	FLinearColor HoverOverUnselectedColor = FLinearColor(1.0f, 0.25f, 0.712328f, .5f); // default green
 
 	/**
 	* Custom color for when geometry is a part of the current selection when using Mesh Element Selection
 	*/
 	UPROPERTY(Config, EditAnywhere, Category = "Modeling Mode|Mesh Element Selection Mode")
-	FLinearColor GeometrySelectedColor = FLinearColor(0.859f, 0.799f, 0.076f, .4f); // default yellow
+	FLinearColor GeometrySelectedColor = FLinearColor(0.982251f, 0.672443f, 0.029557f, .5f); // default yellow
 
 public:
 
