Modeling Tools: Disable frustum selection when tool is active

diff --git a/Engine/Plugins/Editor/ModelingToolsEditorMode/Source/ModelingToolsEditorMode/Private/ModelingToolsEditorMode.cpp b/Engine/Plugins/Editor/ModelingToolsEditorMode/Source/ModelingToolsEditorMode/Private/ModelingToolsEditorMode.cpp
--- a/Engine/Plugins/Editor/ModelingToolsEditorMode/Source/ModelingToolsEditorMode/Private/ModelingToolsEditorMode.cpp
+++ b/Engine/Plugins/Editor/ModelingToolsEditorMode/Source/ModelingToolsEditorMode/Private/ModelingToolsEditorMode.cpp
@@ -1472,6 +1472,11 @@ bool UModelingToolsEditorMode::BoxSelect(FBox& InBox, bool InSelect)
 
 bool UModelingToolsEditorMode::FrustumSelect(const FConvexVolume& InFrustum, FEditorViewportClient* InViewportClient, bool InSelect)
 {
+	if (bIsToolActive)
+	{
+		return true; // will signal that Frustum Select does not need to do anything; disables all frustum select when tool is active
+	}
+
 	if (GetMeshElementSelectionSystemEnabled()
 		&& SelectionManager
 		&& SelectionManager->HasActiveTargets() 
@@ -1527,6 +1532,8 @@ void UModelingToolsEditorMode::OnToolStarted(UInteractiveToolManager* Manager, U
 	// the result. This apparently broken behavior is currently by-design.
 	FSlateThrottleManager::Get().DisableThrottle(true);
 
+	bIsToolActive = true;
+
 	FModelingToolActionCommands::UpdateToolCommandBinding(Tool, Toolkit->GetToolkitCommands(), false);
 	
 	if( FEngineAnalytics::IsAvailable() )
@@ -1549,6 +1556,8 @@ void UModelingToolsEditorMode::OnToolEnded(UInteractiveToolManager* Manager, UIn
 	// re-enable slate throttling (see OnToolStarted)
 	FSlateThrottleManager::Get().DisableThrottle(false);
 
+	bIsToolActive = false;
+
 	FModelingToolActionCommands::UpdateToolCommandBinding(Tool, Toolkit->GetToolkitCommands(), true);
 	
 	// We may require a gizmo location update despite not having changed the selection (transform tool,
diff --git a/Engine/Plugins/Editor/ModelingToolsEditorMode/Source/ModelingToolsEditorMode/Public/ModelingToolsEditorMode.h b/Engine/Plugins/Editor/ModelingToolsEditorMode/Source/ModelingToolsEditorMode/Public/ModelingToolsEditorMode.h
--- a/Engine/Plugins/Editor/ModelingToolsEditorMode/Source/ModelingToolsEditorMode/Public/ModelingToolsEditorMode.h
+++ b/Engine/Plugins/Editor/ModelingToolsEditorMode/Source/ModelingToolsEditorMode/Public/ModelingToolsEditorMode.h
@@ -167,6 +167,8 @@ protected:
 
 	// tracking of unlocked stuff
 	static FDelegateHandle GlobalModelingWorldTeardownEventHandle;
+private:
+	bool bIsToolActive = false;
 };
 
 #if UE_ENABLE_INCLUDE_ORDER_DEPRECATED_IN_5_2
