fix crash on exit by disconnecting stylus api contexts/windows on mode exit, so it does not keep alive editor windows beyond their expected lifetime

diff --git a/Engine/Plugins/Editor/ModelingToolsEditorMode/Source/ModelingToolsEditorMode/Private/ModelingToolsEditorMode.cpp b/Engine/Plugins/Editor/ModelingToolsEditorMode/Source/ModelingToolsEditorMode/Private/ModelingToolsEditorMode.cpp
--- a/Engine/Plugins/Editor/ModelingToolsEditorMode/Source/ModelingToolsEditorMode/Private/ModelingToolsEditorMode.cpp
+++ b/Engine/Plugins/Editor/ModelingToolsEditorMode/Source/ModelingToolsEditorMode/Private/ModelingToolsEditorMode.cpp
@@ -1222,6 +1222,12 @@ void UModelingToolsEditorMode::OnEditorClosed()
 	{
 		GEditor->OnEditorClose().Remove(EditorClosedEventHandle);
 	}
+
+	// cleanup active toolkit stylus input contexts/windows
+	if (Toolkit.IsValid())
+	{
+		static_cast<FModelingToolsEditorModeToolkit*>(Toolkit.Get())->DisconnectStylusStateProviderAPI();
+	}
 }
 
 
diff --git a/Engine/Plugins/Editor/ModelingToolsEditorMode/Source/ModelingToolsEditorMode/Private/ModelingToolsEditorModeToolkit.cpp b/Engine/Plugins/Editor/ModelingToolsEditorMode/Source/ModelingToolsEditorMode/Private/ModelingToolsEditorModeToolkit.cpp
--- a/Engine/Plugins/Editor/ModelingToolsEditorMode/Source/ModelingToolsEditorMode/Private/ModelingToolsEditorModeToolkit.cpp
+++ b/Engine/Plugins/Editor/ModelingToolsEditorMode/Source/ModelingToolsEditorMode/Private/ModelingToolsEditorModeToolkit.cpp
@@ -2721,6 +2721,17 @@ FReply FModelingToolsEditorModeToolkit::HandleCompleteClick()
 	return FReply::Handled();
 }
 
+void FModelingToolsEditorModeToolkit::DisconnectStylusStateProviderAPI()
+{
+#if ENABLE_STYLUS_SUPPORT
+	if (StylusInputHandler)
+	{
+		// replace the input handler with a new, empty version -- disconnecting old windows/contexts
+		StylusInputHandler = MakeUnique<UE::Modeling::FStylusInputHandler>();
+	}
+#endif
+}
+
 IToolStylusStateProviderAPI* FModelingToolsEditorModeToolkit::GetStylusStateProviderAPI() const
 {
 #if ENABLE_STYLUS_SUPPORT
diff --git a/Engine/Plugins/Editor/ModelingToolsEditorMode/Source/ModelingToolsEditorMode/Public/ModelingToolsEditorModeToolkit.h b/Engine/Plugins/Editor/ModelingToolsEditorMode/Source/ModelingToolsEditorMode/Public/ModelingToolsEditorModeToolkit.h
--- a/Engine/Plugins/Editor/ModelingToolsEditorMode/Source/ModelingToolsEditorMode/Public/ModelingToolsEditorModeToolkit.h
+++ b/Engine/Plugins/Editor/ModelingToolsEditorMode/Source/ModelingToolsEditorMode/Public/ModelingToolsEditorModeToolkit.h
@@ -115,6 +115,9 @@ public:
 
 	void NotifySelectionSystemEnabledStateModified();
 
+	// Tells the stylus API to clean up any active tracked windows / contexts
+	void DisconnectStylusStateProviderAPI();
+	
 	IToolStylusStateProviderAPI* GetStylusStateProviderAPI() const;
 
 private:
