Default enable modeling mode's mesh element selections, by changing over to the new settings member variable.
 + use accessor methods so changing between these variables touches less code (in case we need to go back to the old default)

diff --git a/Engine/Plugins/Editor/ModelingToolsEditorMode/Source/ModelingToolsEditorMode/Private/ModelingModeToolkit_Toolbars.cpp b/Engine/Plugins/Editor/ModelingToolsEditorMode/Source/ModelingToolsEditorMode/Private/ModelingModeToolkit_Toolbars.cpp
--- a/Engine/Plugins/Editor/ModelingToolsEditorMode/Source/ModelingToolsEditorMode/Private/ModelingModeToolkit_Toolbars.cpp
+++ b/Engine/Plugins/Editor/ModelingToolsEditorMode/Source/ModelingToolsEditorMode/Private/ModelingModeToolkit_Toolbars.cpp
@@ -139,14 +139,14 @@ void MakeSubMenu_SelectionSupport(FModelingToolsEditorModeToolkit* Toolkit, FMen
 	const FUIAction EnableSelectionsToggle(
 		FExecuteAction::CreateLambda([Toolkit, Settings]
 		{
-			Settings->bEnablePersistentSelections = ! Settings->bEnablePersistentSelections;
+			Settings->SetMeshSelectionsEnabled(!Settings->GetMeshSelectionsEnabled());
 			Settings->SaveConfig();
 			Toolkit->NotifySelectionSystemEnabledStateModified();
 		}),
 		FCanExecuteAction(),
 		FIsActionChecked::CreateLambda([Settings]()
 		{
-			return Settings->bEnablePersistentSelections;
+			return Settings->GetMeshSelectionsEnabled();
 		}));
 	MenuBuilder.AddMenuEntry(LOCTEXT("MeshElementSelectionToggle", "Mesh Element Selection"), 
 		LOCTEXT("MeshElementSelectionToggle_Tooltip", "Enable support for in-viewport Mesh Element Selection in Modeling Mode"),
diff --git a/Engine/Plugins/Editor/ModelingToolsEditorMode/Source/ModelingToolsEditorMode/Private/ModelingToolsEditorMode.cpp b/Engine/Plugins/Editor/ModelingToolsEditorMode/Source/ModelingToolsEditorMode/Private/ModelingToolsEditorMode.cpp
--- a/Engine/Plugins/Editor/ModelingToolsEditorMode/Source/ModelingToolsEditorMode/Private/ModelingToolsEditorMode.cpp
+++ b/Engine/Plugins/Editor/ModelingToolsEditorMode/Source/ModelingToolsEditorMode/Private/ModelingToolsEditorMode.cpp
@@ -285,7 +285,7 @@ void UModelingToolsEditorMode::Enter()
 
 	const UModelingToolsEditorModeSettings* ModelingModeSettings = GetDefault<UModelingToolsEditorModeSettings>();
 	const UModelingToolsModeCustomizationSettings* ModelingEditorSettings = GetDefault<UModelingToolsModeCustomizationSettings>();
-	bSelectionSystemEnabled = ModelingModeSettings->bEnablePersistentSelections;
+	bSelectionSystemEnabled = ModelingModeSettings->GetMeshSelectionsEnabled();
 
 	// Register builders for tool targets that the mode uses.
 	GetInteractiveToolsContext()->TargetManager->AddTargetFactory(NewObject<UStaticMeshComponentToolTargetFactory>(GetToolManager()));
@@ -1332,7 +1332,7 @@ bool UModelingToolsEditorMode::GetMeshElementSelectionSystemEnabled() const
 void UModelingToolsEditorMode::NotifySelectionSystemEnabledStateModified()
 {
 	UModelingToolsEditorModeSettings* Settings = GetMutableDefault<UModelingToolsEditorModeSettings>();
-	bool bNewState = Settings->bEnablePersistentSelections;
+	bool bNewState = Settings->GetMeshSelectionsEnabled();
 	if (bNewState != bSelectionSystemEnabled)
 	{
 		if (bNewState == true)
diff --git a/Engine/Plugins/Editor/ModelingToolsEditorMode/Source/ModelingToolsEditorMode/Public/ModelingToolsEditorModeSettings.h b/Engine/Plugins/Editor/ModelingToolsEditorMode/Source/ModelingToolsEditorMode/Public/ModelingToolsEditorModeSettings.h
--- a/Engine/Plugins/Editor/ModelingToolsEditorMode/Source/ModelingToolsEditorMode/Public/ModelingToolsEditorModeSettings.h
+++ b/Engine/Plugins/Editor/ModelingToolsEditorMode/Source/ModelingToolsEditorMode/Public/ModelingToolsEditorModeSettings.h
@@ -146,20 +146,24 @@ public:
 	// Selection
 	//
 
-	// old preference for mesh selection system that will be disabled in 5.3
-	UPROPERTY(config, EditAnywhere, Category = "Modeling Mode|Selection", meta = (DisplayName="Enable Mesh Selection UI"))
-	bool bEnablePersistentSelections = false;
-
-	// new default-enabled preference for the mesh selection system that we will switch to in 5.3
-	/** Enable/Disable new Mesh Selection System (Experimental). */
+	//~old preference for mesh selection system, to be disabled when we switch to default-enabled mesh element selection
 	//UPROPERTY(config, EditAnywhere, Category = "Modeling Mode|Selection", meta = (DisplayName="Enable Mesh Selection UI"))
 	UPROPERTY()
+	bool bEnablePersistentSelections = false;
+
+	//~new default-enabled preference for the mesh selection system
+	/** Enable/Disable the Mesh Selection System. */
+	UPROPERTY(config, EditAnywhere, Category = "Modeling Mode|Selection", meta = (DisplayName="Enable Mesh Selection UI"))
 	bool bEnableMeshSelections = true;
 
 	virtual bool GetMeshSelectionsEnabled() const
 	{
-		return bEnablePersistentSelections;
-		//return bEnableMeshSelections;
+		return bEnableMeshSelections;
+	}
+	
+	virtual void SetMeshSelectionsEnabled(bool bEnabled)
+	{
+		bEnableMeshSelections = bEnabled;
 	}
 
 
