Modeling Mode: Add Mesh Element Selection to the new viewport toolbar

diff --git a/Engine/Plugins/Editor/ModelingToolsEditorMode/Source/ModelingToolsEditorMode/Private/ModelingModeToolkit_Toolbars.cpp b/Engine/Plugins/Editor/ModelingToolsEditorMode/Source/ModelingToolsEditorMode/Private/ModelingModeToolkit_Toolbars.cpp
--- a/Engine/Plugins/Editor/ModelingToolsEditorMode/Source/ModelingToolsEditorMode/Private/ModelingModeToolkit_Toolbars.cpp
+++ b/Engine/Plugins/Editor/ModelingToolsEditorMode/Source/ModelingToolsEditorMode/Private/ModelingModeToolkit_Toolbars.cpp
@@ -31,6 +31,7 @@
 #include "Framework/Notifications/NotificationManager.h"
 
 #include "ToolMenus.h"
+#include "ViewportToolbar/UnrealEdViewportToolbar.h"
 // to toggle component instance selection
 
 #define LOCTEXT_NAMESPACE "FModelingToolsEditorModeToolkit"
@@ -788,7 +789,9 @@ void FModelingToolsEditorModeToolkit::MakeSelectionPaletteOverlayWidget()
 	{
 		if ( UModelingToolsEditorMode* ModelingMode = Cast<UModelingToolsEditorMode>(GetScriptableEditorMode()) )
 		{
-			return ModelingMode->GetMeshElementSelectionSystemEnabled() ? EVisibility::SelfHitTestInvisible : EVisibility::Collapsed;
+			// only show old MES Toolbar (on left side) when exclusively showing the Old Toolbar (& MES is enabled)
+			// TODO: when new viewport toolbar is finalized remove the ShowNewViewportToolbars condition
+			return ModelingMode->GetMeshElementSelectionSystemEnabled() && !UE::UnrealEd::ShowNewViewportToolbars() ? EVisibility::SelfHitTestInvisible : EVisibility::Collapsed;
 		}
 		return EVisibility::Collapsed;
 	}) );
diff --git a/Engine/Plugins/Editor/ModelingToolsEditorMode/Source/ModelingToolsEditorMode/Private/ModelingToolsEditorMode.cpp b/Engine/Plugins/Editor/ModelingToolsEditorMode/Source/ModelingToolsEditorMode/Private/ModelingToolsEditorMode.cpp
--- a/Engine/Plugins/Editor/ModelingToolsEditorMode/Source/ModelingToolsEditorMode/Private/ModelingToolsEditorMode.cpp
+++ b/Engine/Plugins/Editor/ModelingToolsEditorMode/Source/ModelingToolsEditorMode/Private/ModelingToolsEditorMode.cpp
@@ -136,9 +136,13 @@
 #include "Selection/StaticMeshSelector.h"
 #include "ModelingSelectionInteraction.h"
 #include "DynamicMeshActor.h"
+#include "ModelingToolsEditorModeStyle.h"
 #include "Components/BrushComponent.h"
 #include "Engine/StaticMeshActor.h"
 #include "Engine/World.h" // FWorldDelegates
+#include "ViewportToolbar/UnrealEdViewportToolbar.h"
+#include "ToolMenus.h"
+#include "Framework/Docking/LayoutExtender.h"
 
 #include UE_INLINE_GENERATED_CPP_BY_NAME(ModelingToolsEditorMode)
 
@@ -159,6 +163,8 @@ FDelegateHandle UModelingToolsEditorMode::GlobalModelingWorldTeardownEventHandle
 
 namespace
 {
+const FName ModelingModeOwnerName = "ModelingModeViewportToolbar";
+	
 FString GetToolName(const UInteractiveTool& Tool)
 {
 	const FString* ToolName = FTextInspector::GetSourceString(Tool.GetToolInfo().ToolDisplayName);
@@ -789,7 +795,16 @@ void UModelingToolsEditorMode::Enter()
 					}
 				}
 			}),
-			FCanExecuteAction::CreateLambda([this]() { return (GetToolManager() != nullptr && GetToolManager()->HasAnyActiveTool() == false) ? true : false; }),
+			FCanExecuteAction::CreateLambda([this]()
+			{
+				// necessary due to Exit()->RemoveModelingModeViewportToolbarExtensions->UnregisterOwnerByName->RefreshAllWidgets on next tick
+				// otherwise, this GetToolManager will cause crash because ModeToolsContext is null (after exit)
+				if (!GetInteractiveToolsContext())
+				{
+					return false;
+				}
+				return (GetToolManager() != nullptr && GetToolManager()->HasAnyActiveTool() == false) && GetMeshElementSelectionSystemEnabled() ? true : false;
+			}),
 			FIsActionChecked::CreateLambda([this, TopoMode, ElementMode]() { return (GetSelectionManager() != nullptr && GetSelectionManager()->GetMeshTopologyMode() == TopoMode && GetSelectionManager()->GetSelectionElementType() == ElementMode); }),
 			EUIActionRepeatMode::RepeatDisabled);
 	};
@@ -1036,6 +1051,9 @@ void UModelingToolsEditorMode::Enter()
 		// mode).
 		GetModeManager()->ActivateDefaultMode();
 	});
+
+	// add Mesh Element Selection commands to the new viewport toolbar 
+	PopulateModelingModeViewportToolbar("LevelEditor.ViewportToolbar", Toolkit->GetToolkitCommands());
 }
 
 void UModelingToolsEditorMode::RegisterUVEditor()
@@ -1103,6 +1121,9 @@ void UModelingToolsEditorMode::RegisterUVEditor()
 
 void UModelingToolsEditorMode::Exit()
 {
+	// on Modeling Mode exit, remove Mesh Element Selection commands from new viewport toolbar
+	RemoveModelingModeViewportToolbarExtensions();
+
 	FWorldDelegates::PreLevelRemovedFromWorld.RemoveAll(this);
 	if (BlueprintPreCompileHandle.IsValid())
 	{
@@ -1247,6 +1268,75 @@ void UModelingToolsEditorMode::OnBlueprintPreCompile(UBlueprint* Blueprint)
 	}
 }
 
+ void UModelingToolsEditorMode::PopulateModelingModeViewportToolbar(const FName InMenuName, const TSharedPtr<const FUICommandList>& InCommandList)
+{
+	FToolMenuOwnerScoped ScopeOwner(ModelingModeOwnerName);
+	UToolMenu* const Menu = UToolMenus::Get()->ExtendMenu(InMenuName);
+	
+	
+	FToolMenuSection& MeshElementSelectionMenu = Menu->FindOrAddSection("Left");
+	FToolMenuEntry Entry = FToolMenuEntry::InitSubMenu(
+		"Mesh Element Selection",
+		LOCTEXT("MeshElementSelectionSubmenuLabel", "Mesh Element Selection"),
+		LOCTEXT("MeshElementSelectionSubmenuTooltip", "Mesh Element Selection settings in the viewport"),
+		FNewToolMenuDelegate::CreateLambda([ WeakCommandList = InCommandList.ToWeakPtr()](UToolMenu* Submenu) -> void
+		{
+			FToolMenuSection& MeshElementSelectionOptionsSection =
+			Submenu->FindOrAddSection("Mesh Element Selection", LOCTEXT("MeshElementSelectionLabel", "Mesh Element Selection"));
+			if (const TSharedPtr<const FUICommandList>& CommandList = WeakCommandList.Pin())
+			{
+				auto CreateToolMenuEntry = [&CommandList, &MeshElementSelectionOptionsSection](const TSharedPtr<FUICommandInfo>& Command)
+				{
+					FToolMenuEntry MESMode = FToolMenuEntry::InitMenuEntryWithCommandList( Command, CommandList );
+					MESMode.SetShowInToolbarTopLevel(true); // add button to toolbar
+					MeshElementSelectionOptionsSection.AddEntry(MESMode);
+				};
+
+				// add all Mesh Element Selection modes
+				CreateToolMenuEntry(FModelingToolsManagerCommands::Get().MeshSelectionModeAction_NoSelection);
+				CreateToolMenuEntry(FModelingToolsManagerCommands::Get().MeshSelectionModeAction_MeshTriangles);
+				CreateToolMenuEntry(FModelingToolsManagerCommands::Get().MeshSelectionModeAction_MeshEdges);
+				CreateToolMenuEntry(FModelingToolsManagerCommands::Get().MeshSelectionModeAction_MeshVertices);
+				CreateToolMenuEntry(FModelingToolsManagerCommands::Get().MeshSelectionModeAction_GroupFaces);
+				CreateToolMenuEntry(FModelingToolsManagerCommands::Get().MeshSelectionModeAction_GroupEdges);
+				CreateToolMenuEntry(FModelingToolsManagerCommands::Get().MeshSelectionModeAction_GroupCorners);
+
+				// Section for Selection Edits
+				{
+					FToolMenuSection& SelectionEditsSection =
+						Submenu->FindOrAddSection("SelectionEdits", LOCTEXT("SelectionEditsLabel", "Selection Edits"));
+
+					SelectionEditsSection.AddEntry(
+					FToolMenuEntry::InitMenuEntryWithCommandList(
+						FModelingToolsManagerCommands::Get().BeginSelectionAction_SelectAll, CommandList));
+					SelectionEditsSection.AddEntry(
+					FToolMenuEntry::InitMenuEntryWithCommandList(
+						FModelingToolsManagerCommands::Get().BeginSelectionAction_Invert, CommandList));
+					SelectionEditsSection.AddEntry(
+					FToolMenuEntry::InitMenuEntryWithCommandList(
+						FModelingToolsManagerCommands::Get().BeginSelectionAction_InvertConnected, CommandList));
+					SelectionEditsSection.AddEntry(
+					FToolMenuEntry::InitMenuEntryWithCommandList(
+						FModelingToolsManagerCommands::Get().BeginSelectionAction_Expand, CommandList));
+					SelectionEditsSection.AddEntry(
+					FToolMenuEntry::InitMenuEntryWithCommandList(
+						FModelingToolsManagerCommands::Get().BeginSelectionAction_Contract, CommandList));
+				}
+			}
+		})
+	);
+	
+	Entry.ToolBarData.LabelOverride = FText();
+	Entry.InsertPosition = FToolMenuInsert("Transform", EToolMenuInsertType::After);
+	Entry.Icon = FSlateIcon(FModelingToolsEditorModeStyle::GetStyleSetName(), "ModelingToolsManagerCommands.MeshElementSelection");
+
+	MeshElementSelectionMenu.AddEntry(Entry);
+}
+
+void UModelingToolsEditorMode::RemoveModelingModeViewportToolbarExtensions()
+{
+	UToolMenus::Get()->UnregisterOwnerByName(ModelingModeOwnerName);
+}
 
 void UModelingToolsEditorMode::OnToolsContextRender(IToolsContextRenderAPI* RenderAPI)
 {
diff --git a/Engine/Plugins/Editor/ModelingToolsEditorMode/Source/ModelingToolsEditorMode/Private/ModelingToolsEditorModeStyle.cpp b/Engine/Plugins/Editor/ModelingToolsEditorMode/Source/ModelingToolsEditorMode/Private/ModelingToolsEditorModeStyle.cpp
--- a/Engine/Plugins/Editor/ModelingToolsEditorMode/Source/ModelingToolsEditorMode/Private/ModelingToolsEditorModeStyle.cpp
+++ b/Engine/Plugins/Editor/ModelingToolsEditorMode/Source/ModelingToolsEditorMode/Private/ModelingToolsEditorModeStyle.cpp
@@ -370,6 +370,8 @@ void FModelingToolsEditorModeStyle::Initialize()
 		StyleSet->Set("ModelingToolsManagerCommands.MeshSelectionModeAction_GroupCorners", new IMAGE_BRUSH_SVG("Icons/ModSelectionPolygroupVerts_16", Icon20x20));
 		StyleSet->Set("ModelingToolsManagerCommands.MeshSelectionModeAction_GroupEdges", new IMAGE_BRUSH_SVG("Icons/ModSelectionPolygroupEdges_16", Icon20x20));
 
+		StyleSet->Set("ModelingToolsManagerCommands.MeshElementSelection", new IMAGE_BRUSH_SVG("Icons/MeshSelect_16", Icon20x20));
+
 		StyleSet->Set("ModelingToolsManagerCommands.BeginSelectionAction_SelectAll", new IMAGE_BRUSH_SVG("Icons/ModSelectionAll_16", Icon20x20));
 		StyleSet->Set("ModelingToolsManagerCommands.BeginSelectionAction_Invert", new IMAGE_BRUSH_SVG("Icons/ModSelectionInverse_16", Icon20x20));
 		StyleSet->Set("ModelingToolsManagerCommands.BeginSelectionAction_ExpandToConnected", new IMAGE_BRUSH_SVG("Icons/ModSelectionConnected_16", Icon20x20));
diff --git a/Engine/Plugins/Editor/ModelingToolsEditorMode/Source/ModelingToolsEditorMode/Private/ModelingToolsManagerActions.cpp b/Engine/Plugins/Editor/ModelingToolsEditorMode/Source/ModelingToolsEditorMode/Private/ModelingToolsManagerActions.cpp
--- a/Engine/Plugins/Editor/ModelingToolsEditorMode/Source/ModelingToolsEditorMode/Private/ModelingToolsManagerActions.cpp
+++ b/Engine/Plugins/Editor/ModelingToolsEditorMode/Source/ModelingToolsEditorMode/Private/ModelingToolsManagerActions.cpp
@@ -71,6 +71,9 @@ void FModelingToolsManagerCommands::RegisterCommands()
 #define REGISTER_MODELING_TOOL_COMMAND(ToolCommandInfo, ToolName, ToolTip) \
 		UI_COMMAND(ToolCommandInfo, ToolName, ToolTip, EUserInterfaceActionType::ToggleButton, FInputChord()); \
 		RegisteredTools.Add(FStartToolCommand{ ToolName, ToolCommandInfo });
+#define REGISTER_MODELING_TOOL_COMMAND_RADIO(ToolCommandInfo, ToolName, ToolTip) \
+		UI_COMMAND(ToolCommandInfo, ToolName, ToolTip, EUserInterfaceActionType::RadioButton, FInputChord()); \
+		RegisteredTools.Add(FStartToolCommand{ ToolName, ToolCommandInfo });
 #define REGISTER_MODELING_TOOL_COMMAND_WITH_SHORTNAME(ToolCommandInfo, ToolName, ShortNameFText, ToolTip) \
 		UI_COMMAND(ToolCommandInfo, ToolName, ToolTip, EUserInterfaceActionType::ToggleButton, FInputChord()); \
 		RegisteredTools.Add(FStartToolCommand{ ToolName, ToolCommandInfo, ShortNameFText });
@@ -230,13 +233,13 @@ void FModelingToolsManagerCommands::RegisterCommands()
 	REGISTER_MODELING_TOOL_COMMAND_WITH_SHORTNAME(BeginPolyModelTool_TriSel, "Tri Select", LOCTEXT("BeginPolyModelTool_TriSel_short", "TriSel"), "Select / Edit the current Mesh triangles with a brush interface");
 
 
-	REGISTER_MODELING_TOOL_COMMAND(MeshSelectionModeAction_NoSelection, "None", "Disable Mesh Element Selection (Select Objects Only)");
-	REGISTER_MODELING_TOOL_COMMAND(MeshSelectionModeAction_MeshTriangles, "Tris", "Select Mesh Triangles");
-	REGISTER_MODELING_TOOL_COMMAND(MeshSelectionModeAction_MeshVertices, "Verts", "Select Mesh Vertices");
-	REGISTER_MODELING_TOOL_COMMAND(MeshSelectionModeAction_MeshEdges, "Edges", "Select Mesh Edges");
-	REGISTER_MODELING_TOOL_COMMAND(MeshSelectionModeAction_GroupFaces, "Groups", "Select Mesh PolyGroups");
-	REGISTER_MODELING_TOOL_COMMAND(MeshSelectionModeAction_GroupCorners, "Corners", "Select Mesh PolyGroup Corners/Vertices");
-	REGISTER_MODELING_TOOL_COMMAND(MeshSelectionModeAction_GroupEdges, "Borders", "Select Mesh PolyGroup Borders/Edges");
+	REGISTER_MODELING_TOOL_COMMAND_RADIO(MeshSelectionModeAction_NoSelection, "Object Selection", "Disable Mesh Element Selection (Select Objects Only)");
+	REGISTER_MODELING_TOOL_COMMAND_RADIO(MeshSelectionModeAction_MeshTriangles, "Mesh Triangles Selection", "Select Mesh Triangles");
+	REGISTER_MODELING_TOOL_COMMAND_RADIO(MeshSelectionModeAction_MeshVertices, "Mesh Vertices Selection", "Select Mesh Vertices");
+	REGISTER_MODELING_TOOL_COMMAND_RADIO(MeshSelectionModeAction_MeshEdges, "Mesh Edges Selection", "Select Mesh Edges");
+	REGISTER_MODELING_TOOL_COMMAND_RADIO(MeshSelectionModeAction_GroupFaces, "PolyGroup Selection", "Select Mesh PolyGroups");
+	REGISTER_MODELING_TOOL_COMMAND_RADIO(MeshSelectionModeAction_GroupCorners, "PolyGroup Corners Selection", "Select Mesh PolyGroup Corners/Vertices");
+	REGISTER_MODELING_TOOL_COMMAND_RADIO(MeshSelectionModeAction_GroupEdges, "PolyGroup Borders Selection", "Select Mesh PolyGroup Borders/Edges");
 
 	REGISTER_MODELING_ACTION_COMMAND(BeginSelectionAction_Delete, "Delete", "Delete the current Mesh Selection");
 	REGISTER_MODELING_ACTION_COMMAND(BeginSelectionAction_Retriangulate, "Clean", "Retriangulate the current Mesh or Mesh Selection");
diff --git a/Engine/Plugins/Editor/ModelingToolsEditorMode/Source/ModelingToolsEditorMode/Public/ModelingToolsEditorMode.h b/Engine/Plugins/Editor/ModelingToolsEditorMode/Source/ModelingToolsEditorMode/Public/ModelingToolsEditorMode.h
--- a/Engine/Plugins/Editor/ModelingToolsEditorMode/Source/ModelingToolsEditorMode/Public/ModelingToolsEditorMode.h
+++ b/Engine/Plugins/Editor/ModelingToolsEditorMode/Source/ModelingToolsEditorMode/Public/ModelingToolsEditorMode.h
@@ -169,6 +169,11 @@ protected:
 	static FDelegateHandle GlobalModelingWorldTeardownEventHandle;
 private:
 	bool bIsToolActive = false;
+
+	// add modeling-mode-specific portion of new viewport toolbar
+	static void PopulateModelingModeViewportToolbar(const FName InMenuName, const TSharedPtr<const FUICommandList>& InCommandList);
+	// remove modeling-mode-specific portion of new viewport toolbar
+	static void RemoveModelingModeViewportToolbarExtensions();
 };
 
 #if UE_ENABLE_INCLUDE_ORDER_DEPRECATED_IN_5_2
