Modeling Mode: minor UI updates to Mesh Element Selection in new viewport toolbar

diff --git a/Engine/Plugins/Editor/ModelingToolsEditorMode/Source/ModelingToolsEditorMode/Private/ModelingToolsEditorMode.cpp b/Engine/Plugins/Editor/ModelingToolsEditorMode/Source/ModelingToolsEditorMode/Private/ModelingToolsEditorMode.cpp
--- a/Engine/Plugins/Editor/ModelingToolsEditorMode/Source/ModelingToolsEditorMode/Private/ModelingToolsEditorMode.cpp
+++ b/Engine/Plugins/Editor/ModelingToolsEditorMode/Source/ModelingToolsEditorMode/Private/ModelingToolsEditorMode.cpp
@@ -1282,7 +1282,7 @@ void UModelingToolsEditorMode::OnBlueprintPreCompile(UBlueprint* Blueprint)
 		FNewToolMenuDelegate::CreateLambda([ WeakCommandList = InCommandList.ToWeakPtr()](UToolMenu* Submenu) -> void
 		{
 			FToolMenuSection& MeshElementSelectionOptionsSection =
-			Submenu->FindOrAddSection("Mesh Element Selection", LOCTEXT("MeshElementSelectionLabel", "Mesh Element Selection"));
+			Submenu->FindOrAddSection("Element Selection", LOCTEXT("ElementSelectionLabel", "Element Selection"));
 			if (const TSharedPtr<const FUICommandList>& CommandList = WeakCommandList.Pin())
 			{
 				auto CreateToolMenuEntry = [&CommandList, &MeshElementSelectionOptionsSection](const TSharedPtr<FUICommandInfo>& Command)
@@ -1325,7 +1325,8 @@ void UModelingToolsEditorMode::OnBlueprintPreCompile(UBlueprint* Blueprint)
 			}
 		})
 	);
-	
+
+	Entry.ToolBarData.ResizeParams.ClippingPriority = 950;
 	Entry.ToolBarData.LabelOverride = FText();
 	Entry.InsertPosition = FToolMenuInsert("Transform", EToolMenuInsertType::After);
 	Entry.Icon = FSlateIcon(FModelingToolsEditorModeStyle::GetStyleSetName(), "ModelingToolsManagerCommands.MeshElementSelection");
diff --git a/Engine/Plugins/Editor/ModelingToolsEditorMode/Source/ModelingToolsEditorMode/Private/ModelingToolsManagerActions.cpp b/Engine/Plugins/Editor/ModelingToolsEditorMode/Source/ModelingToolsEditorMode/Private/ModelingToolsManagerActions.cpp
--- a/Engine/Plugins/Editor/ModelingToolsEditorMode/Source/ModelingToolsEditorMode/Private/ModelingToolsManagerActions.cpp
+++ b/Engine/Plugins/Editor/ModelingToolsEditorMode/Source/ModelingToolsEditorMode/Private/ModelingToolsManagerActions.cpp
@@ -233,13 +233,13 @@ void FModelingToolsManagerCommands::RegisterCommands()
 	REGISTER_MODELING_TOOL_COMMAND_WITH_SHORTNAME(BeginPolyModelTool_TriSel, "Tri Select", LOCTEXT("BeginPolyModelTool_TriSel_short", "TriSel"), "Select / Edit the current Mesh triangles with a brush interface");
 
 
-	REGISTER_MODELING_TOOL_COMMAND_RADIO(MeshSelectionModeAction_NoSelection, "Object Selection", "Disable Mesh Element Selection (Select Objects Only)");
-	REGISTER_MODELING_TOOL_COMMAND_RADIO(MeshSelectionModeAction_MeshTriangles, "Mesh Triangles Selection", "Select Mesh Triangles");
-	REGISTER_MODELING_TOOL_COMMAND_RADIO(MeshSelectionModeAction_MeshVertices, "Mesh Vertices Selection", "Select Mesh Vertices");
-	REGISTER_MODELING_TOOL_COMMAND_RADIO(MeshSelectionModeAction_MeshEdges, "Mesh Edges Selection", "Select Mesh Edges");
-	REGISTER_MODELING_TOOL_COMMAND_RADIO(MeshSelectionModeAction_GroupFaces, "PolyGroup Selection", "Select Mesh PolyGroups");
-	REGISTER_MODELING_TOOL_COMMAND_RADIO(MeshSelectionModeAction_GroupCorners, "PolyGroup Corners Selection", "Select Mesh PolyGroup Corners/Vertices");
-	REGISTER_MODELING_TOOL_COMMAND_RADIO(MeshSelectionModeAction_GroupEdges, "PolyGroup Borders Selection", "Select Mesh PolyGroup Borders/Edges");
+	REGISTER_MODELING_TOOL_COMMAND_RADIO(MeshSelectionModeAction_NoSelection, "Object", "Disable Mesh Element Selection (Select Objects Only)");
+	REGISTER_MODELING_TOOL_COMMAND_RADIO(MeshSelectionModeAction_MeshTriangles, "Mesh Triangles", "Select Mesh Triangles");
+	REGISTER_MODELING_TOOL_COMMAND_RADIO(MeshSelectionModeAction_MeshVertices, "Mesh Vertices", "Select Mesh Vertices");
+	REGISTER_MODELING_TOOL_COMMAND_RADIO(MeshSelectionModeAction_MeshEdges, "Mesh Edges", "Select Mesh Edges");
+	REGISTER_MODELING_TOOL_COMMAND_RADIO(MeshSelectionModeAction_GroupFaces, "PolyGroups", "Select Mesh PolyGroups");
+	REGISTER_MODELING_TOOL_COMMAND_RADIO(MeshSelectionModeAction_GroupCorners, "PolyGroup Corners", "Select Mesh PolyGroup Corners/Vertices");
+	REGISTER_MODELING_TOOL_COMMAND_RADIO(MeshSelectionModeAction_GroupEdges, "PolyGroup Borders", "Select Mesh PolyGroup Borders/Edges");
 
 	REGISTER_MODELING_ACTION_COMMAND(BeginSelectionAction_Delete, "Delete", "Delete the current Mesh Selection");
 	REGISTER_MODELING_ACTION_COMMAND(BeginSelectionAction_Retriangulate, "Clean", "Retriangulate the current Mesh or Mesh Selection");
