Modling Mode: Update MeshElementSelection New Viewport Toolbar UI

diff --git a/Engine/Plugins/Editor/ModelingToolsEditorMode/Source/ModelingToolsEditorMode/Private/ModelingToolsEditorMode.cpp b/Engine/Plugins/Editor/ModelingToolsEditorMode/Source/ModelingToolsEditorMode/Private/ModelingToolsEditorMode.cpp
--- a/Engine/Plugins/Editor/ModelingToolsEditorMode/Source/ModelingToolsEditorMode/Private/ModelingToolsEditorMode.cpp
+++ b/Engine/Plugins/Editor/ModelingToolsEditorMode/Source/ModelingToolsEditorMode/Private/ModelingToolsEditorMode.cpp
@@ -1294,9 +1294,11 @@ void UModelingToolsEditorMode::OnBlueprintPreCompile(UBlueprint* Blueprint)
 
 				// add all Mesh Element Selection modes
 				CreateToolMenuEntry(FModelingToolsManagerCommands::Get().MeshSelectionModeAction_NoSelection);
+				MeshElementSelectionOptionsSection.AddSeparator(NAME_None);
 				CreateToolMenuEntry(FModelingToolsManagerCommands::Get().MeshSelectionModeAction_MeshTriangles);
 				CreateToolMenuEntry(FModelingToolsManagerCommands::Get().MeshSelectionModeAction_MeshEdges);
 				CreateToolMenuEntry(FModelingToolsManagerCommands::Get().MeshSelectionModeAction_MeshVertices);
+				MeshElementSelectionOptionsSection.AddSeparator(NAME_None);
 				CreateToolMenuEntry(FModelingToolsManagerCommands::Get().MeshSelectionModeAction_GroupFaces);
 				CreateToolMenuEntry(FModelingToolsManagerCommands::Get().MeshSelectionModeAction_GroupEdges);
 				CreateToolMenuEntry(FModelingToolsManagerCommands::Get().MeshSelectionModeAction_GroupCorners);
