InteractiveToolsFramework: Add support for hotkeys in tool extensions.

diff --git a/Engine/Plugins/Editor/ModelingToolsEditorMode/Source/ModelingToolsEditorMode/Private/ModelingToolsEditorMode.cpp b/Engine/Plugins/Editor/ModelingToolsEditorMode/Source/ModelingToolsEditorMode/Private/ModelingToolsEditorMode.cpp
--- a/Engine/Plugins/Editor/ModelingToolsEditorMode/Source/ModelingToolsEditorMode/Private/ModelingToolsEditorMode.cpp
+++ b/Engine/Plugins/Editor/ModelingToolsEditorMode/Source/ModelingToolsEditorMode/Private/ModelingToolsEditorMode.cpp
@@ -161,7 +161,7 @@ FDateTime UModelingToolsEditorMode::LastToolStartTimestamp;
 
 FDelegateHandle UModelingToolsEditorMode::GlobalModelingWorldTeardownEventHandle;
 
-namespace
+namespace ModelingToolsEditorModeLocals
 {
 const FName ModelingModeOwnerName = "ModelingModeViewportToolbar";
 	
@@ -287,6 +287,8 @@ void UModelingToolsEditorMode::Tick(FEditorViewportClient* ViewportClient, float
 
 void UModelingToolsEditorMode::Enter()
 {
+	using namespace ModelingToolsEditorModeLocals;
+
 	UEdMode::Enter();
 
 	const UModelingToolsEditorModeSettings* ModelingModeSettings = GetDefault<UModelingToolsEditorModeSettings>();
@@ -870,6 +872,7 @@ void UModelingToolsEditorMode::Enter()
 				UE_LOG(LogTemp, Log, TEXT("%s - Registering Tool [%s]"), *ExtensionPrefix, *ToolInfo.ToolName.ToString());
 
 				RegisterTool(ToolInfo.ToolCommand, ToolInfo.ToolName.ToString(), ToolInfo.ToolBuilder);
+				ExtensionToolToInfo.Add(ToolInfo.ToolName.ToString(), ToolInfo);
 			}
 
 			TArray<TSubclassOf<UToolTargetFactory>> ExtensionToolTargetFactoryClasses;
@@ -1270,6 +1273,8 @@ void UModelingToolsEditorMode::OnBlueprintPreCompile(UBlueprint* Blueprint)
 
  void UModelingToolsEditorMode::PopulateModelingModeViewportToolbar(const FName InMenuName, const TSharedPtr<const FUICommandList>& InCommandList)
 {
+	 using namespace ModelingToolsEditorModeLocals;
+
 	FToolMenuOwnerScoped ScopeOwner(ModelingModeOwnerName);
 	UToolMenu* const Menu = UToolMenus::Get()->ExtendMenu(InMenuName);
 	
@@ -1338,7 +1343,7 @@ void UModelingToolsEditorMode::OnBlueprintPreCompile(UBlueprint* Blueprint)
 
 void UModelingToolsEditorMode::RemoveModelingModeViewportToolbarExtensions()
 {
-	UToolMenus::Get()->UnregisterOwnerByName(ModelingModeOwnerName);
+	UToolMenus::Get()->UnregisterOwnerByName(ModelingToolsEditorModeLocals::ModelingModeOwnerName);
 }
 
 void UModelingToolsEditorMode::OnToolsContextRender(IToolsContextRenderAPI* RenderAPI)
@@ -1626,6 +1631,8 @@ void UModelingToolsEditorMode::OnToolPostBuild(
 
 void UModelingToolsEditorMode::OnToolStarted(UInteractiveToolManager* Manager, UInteractiveTool* Tool)
 {
+	using namespace ModelingToolsEditorModeLocals;
+
 	// disable slate throttling so that Tool background computes responding to sliders can properly be processed
 	// on Tool Tick. Otherwise, when a Tool kicks off a background update in a background thread, the computed
 	// result will be ignored until the user moves the slider, ie you cannot hold down the mouse and wait to see
@@ -1635,7 +1642,11 @@ void UModelingToolsEditorMode::OnToolStarted(UInteractiveToolManager* Manager, U
 	bIsToolActive = true;
 
 	FModelingToolActionCommands::UpdateToolCommandBinding(Tool, Toolkit->GetToolkitCommands(), false);
-	
+	if (TryGetExtensionToolCommandGetter(Manager, Tool, ExtensionToolCommandsGetter) && ensure(ExtensionToolCommandsGetter))
+	{
+		ExtensionToolCommandsGetter().BindCommandsForCurrentTool(Toolkit->GetToolkitCommands(), Tool);
+	}
+
 	if( FEngineAnalytics::IsAvailable() )
 	{
 		FEngineAnalytics::GetProvider().RecordEvent(TEXT("Editor.Usage.MeshModelingMode.ToolStarted"),
@@ -1653,12 +1664,19 @@ void UModelingToolsEditorMode::OnToolStarted(UInteractiveToolManager* Manager, U
 
 void UModelingToolsEditorMode::OnToolEnded(UInteractiveToolManager* Manager, UInteractiveTool* Tool)
 {
+	using namespace ModelingToolsEditorModeLocals;
+
 	// re-enable slate throttling (see OnToolStarted)
 	FSlateThrottleManager::Get().DisableThrottle(false);
 
 	bIsToolActive = false;
 
 	FModelingToolActionCommands::UpdateToolCommandBinding(Tool, Toolkit->GetToolkitCommands(), true);
+	if (ExtensionToolCommandsGetter)
+	{
+		ExtensionToolCommandsGetter().UnbindActiveCommands(Toolkit->GetToolkitCommands());
+		ExtensionToolCommandsGetter = nullptr;
+	}
 	
 	// We may require a gizmo location update despite not having changed the selection (transform tool,
 	// edit pivot, etc).
@@ -1679,6 +1697,30 @@ void UModelingToolsEditorMode::OnToolEnded(UInteractiveToolManager* Manager, UIn
 	}
 }
 
+bool UModelingToolsEditorMode::TryGetExtensionToolCommandGetter(
+	UInteractiveToolManager* Manager, UInteractiveTool* Tool, 
+	TFunction<const UE::IInteractiveToolCommandsInterface& ()>& GetterOut) const
+{
+	if (!ensure(Manager && Tool)
+		|| Manager->GetActiveTool(EToolSide::Mouse) != Tool)
+	{
+		return false;
+	}
+
+	FString ToolName = Manager->GetActiveToolName(EToolSide::Mouse);
+	if (ToolName.IsEmpty())
+	{
+		return false;
+	}
+	const FExtensionToolDescription* ToolDescription = ExtensionToolToInfo.Find(ToolName);
+	if (!ToolDescription || !ToolDescription->ToolCommandsGetter)
+	{
+		return false;
+	}
+	GetterOut = ToolDescription->ToolCommandsGetter;
+	return true;
+}
+
 void UModelingToolsEditorMode::BindCommands()
 {
 	Super::BindCommands();
diff --git a/Engine/Plugins/Editor/ModelingToolsEditorMode/Source/ModelingToolsEditorMode/Public/ModelingModeToolExtensions.h b/Engine/Plugins/Editor/ModelingToolsEditorMode/Source/ModelingToolsEditorMode/Public/ModelingModeToolExtensions.h
--- a/Engine/Plugins/Editor/ModelingToolsEditorMode/Source/ModelingToolsEditorMode/Public/ModelingModeToolExtensions.h
+++ b/Engine/Plugins/Editor/ModelingToolsEditorMode/Source/ModelingToolsEditorMode/Public/ModelingModeToolExtensions.h
@@ -9,6 +9,10 @@
 class UInteractiveToolsContext;
 class FUICommandInfo;
 class UInteractiveToolBuilder;
+namespace UE
+{
+	class IInteractiveToolCommandsInterface;
+}
 
 
 // dummy class for extension API below
@@ -45,6 +49,12 @@ struct MODELINGTOOLSEDITORMODE_API FExtensionToolDescription
 	TSharedPtr<FUICommandInfo> ToolCommand;
 	/** Builder for the Tool. This can be null if FExtensionToolQueryInfo.bIsInfoQueryOnly is true. */
 	UInteractiveToolBuilder* ToolBuilder = nullptr;
+	/**
+	 * Optional lambda that returns this tool's hotkey commands object, if one exists, so that
+	 *  those commands can be properly bound/unbound when tool starts/ends. E.g. this will be
+	 *  something like "return FThisExtensionToolCommands::Get();"
+	 */
+	TFunction<const UE::IInteractiveToolCommandsInterface&()> ToolCommandsGetter = nullptr;
 };
 
 /**
diff --git a/Engine/Plugins/Editor/ModelingToolsEditorMode/Source/ModelingToolsEditorMode/Public/ModelingToolsEditorMode.h b/Engine/Plugins/Editor/ModelingToolsEditorMode/Source/ModelingToolsEditorMode/Public/ModelingToolsEditorMode.h
--- a/Engine/Plugins/Editor/ModelingToolsEditorMode/Source/ModelingToolsEditorMode/Public/ModelingToolsEditorMode.h
+++ b/Engine/Plugins/Editor/ModelingToolsEditorMode/Source/ModelingToolsEditorMode/Public/ModelingToolsEditorMode.h
@@ -4,9 +4,14 @@
 
 #include "Tools/LegacyEdModeWidgetHelpers.h"
 
+#include "ModelingModeToolExtensions.h" // FExtensionToolDescription
 
 #include "ModelingToolsEditorMode.generated.h"
 
+namespace UE
+{
+	class IInteractiveToolCommandsInterface;
+}
 class IToolsContextRenderAPI;
 enum class EModelingModeActionCommands;
 enum class EToolSide;
@@ -174,6 +179,15 @@ private:
 	static void PopulateModelingModeViewportToolbar(const FName InMenuName, const TSharedPtr<const FUICommandList>& InCommandList);
 	// remove modeling-mode-specific portion of new viewport toolbar
 	static void RemoveModelingModeViewportToolbarExtensions();
+
+	// Support extension tools having their own hotkey classes
+	TMap<FString, FExtensionToolDescription> ExtensionToolToInfo;
+	// Note: this will only work when the given tool is active, because we get the tool identifier
+	//  out of the manager using GetActiveToolName
+	bool TryGetExtensionToolCommandGetter(UInteractiveToolManager* Manager, UInteractiveTool* Tool, 
+		TFunction<const UE::IInteractiveToolCommandsInterface&()>& GetterOut) const;
+	// Used to unbind extension tool commands
+	TFunction<const UE::IInteractiveToolCommandsInterface& ()> ExtensionToolCommandsGetter;
 };
 
 #if UE_ENABLE_INCLUDE_ORDER_DEPRECATED_IN_5_2
diff --git a/Engine/Source/Editor/UnrealEd/Public/Tools/InteractiveToolsCommands.h b/Engine/Source/Editor/UnrealEd/Public/Tools/InteractiveToolsCommands.h
--- a/Engine/Source/Editor/UnrealEd/Public/Tools/InteractiveToolsCommands.h
+++ b/Engine/Source/Editor/UnrealEd/Public/Tools/InteractiveToolsCommands.h
@@ -9,7 +9,28 @@
 #include "StandardToolModeCommands.h"
 #include "Styling/ISlateStyle.h"
 
-
+namespace UE
+{
+	/**
+	 * Interface for functions that get called on command objects on tool start/end. Used to allow the caller to
+	 *  not know the class at compile time, so that, for instance, tools injected via a plugin can still have
+	 *  hotkey support.
+	 * The command object's RegisterCommands() function still needs to be called in the module's startup method,
+	 *  but then this interface can be provided to the host to call when the tool starts/ends.
+	 */
+	class IInteractiveToolCommandsInterface
+	{
+	public:
+		/**
+		 * Bind any of the registered UICommands to the given Tool, if they are compatible.
+		 */
+		virtual void BindCommandsForCurrentTool(TSharedPtr<FUICommandList> UICommandList, UInteractiveTool* Tool) const = 0;
+		/**
+		 * Unbind all of the currently-registered commands
+		 */
+		virtual void UnbindActiveCommands(TSharedPtr<FUICommandList> UICommandList) const = 0;
+	};
+}
 
 /**
  * TInteractiveToolCommands is a base class that handles connecting up Tool Actions 
@@ -25,7 +46,8 @@
  *    - in your EdModeImpl::InputKey override, call UICommandList->ProcessCommandBindings()
  */
 template<typename CommandContextType>
-class TInteractiveToolCommands : public TCommands<CommandContextType>
+class TInteractiveToolCommands : public TCommands<CommandContextType>, 
+	public UE::IInteractiveToolCommandsInterface
 {
 
 public:
@@ -39,12 +61,12 @@ public:
 	/**
 	 * Bind any of the registered UICommands to the given Tool, if they are compatible.
 	 */
-	virtual void BindCommandsForCurrentTool(TSharedPtr<FUICommandList> UICommandList, UInteractiveTool* Tool) const;
+	virtual void BindCommandsForCurrentTool(TSharedPtr<FUICommandList> UICommandList, UInteractiveTool* Tool) const override;
 
 	/**
 	 * Unbind all of the currently-registered commands
 	 */
-	virtual void UnbindActiveCommands(TSharedPtr<FUICommandList> UICommandList) const;
+	virtual void UnbindActiveCommands(TSharedPtr<FUICommandList> UICommandList) const override;
 
 
 	//
