MegaMesh: Add Plane Height Sculpt and Sphere Height Sculpt tools for sculpting landscapes or planets. Add support for hotkeys in the plugin.

diff --git a/Engine/Plugins/Editor/ModelingToolsEditorMode/Source/ModelingToolsEditorMode/Private/ModelingToolsEditorModeStyle.cpp b/Engine/Plugins/Editor/ModelingToolsEditorMode/Source/ModelingToolsEditorMode/Private/ModelingToolsEditorModeStyle.cpp
--- a/Engine/Plugins/Editor/ModelingToolsEditorMode/Source/ModelingToolsEditorMode/Private/ModelingToolsEditorModeStyle.cpp
+++ b/Engine/Plugins/Editor/ModelingToolsEditorMode/Source/ModelingToolsEditorMode/Private/ModelingToolsEditorModeStyle.cpp
@@ -477,6 +477,11 @@ void FModelingToolsEditorModeStyle::Initialize()
 		StyleSet->Set("BrushTypeIcons.Grab", new IMAGE_BRUSH_SVG("Icons/BrushIcons/Brush_Grab", Icon120));
 		StyleSet->Set("BrushTypeIcons.GrabSharp", new IMAGE_BRUSH_SVG("Icons/BrushIcons/Brush_GrabSharp", Icon120));
 		StyleSet->Set("BrushTypeIcons.Twist", new IMAGE_BRUSH_SVG("Icons/BrushIcons/Brush_Twist", Icon120));
+		// Height brushes (currently using same icons as corresponding landscape brushes)
+		StyleSet->Set("BrushTypeIcons.HeightSculpt", new FSlateImageBrush(StyleSet->RootToCoreContentDir(TEXT("../Editor/Slate/Icons/Landscape/Landscape_Sculpt_x40.png")), Icon20x20));
+		StyleSet->Set("BrushTypeIcons.HeightSmooth", new FSlateImageBrush(StyleSet->RootToCoreContentDir(TEXT("../Editor/Slate/Icons/Landscape/Landscape_Smooth_x40.png")), Icon20x20));
+		StyleSet->Set("BrushTypeIcons.HeightFlatten", new FSlateImageBrush(StyleSet->RootToCoreContentDir(TEXT("../Editor/Slate/Icons/Landscape/Landscape_Flatten_x40.png")), Icon20x20));
+		StyleSet->Set("BrushTypeIcons.HeightErase", new FSlateImageBrush(StyleSet->RootToCoreContentDir(TEXT("../Editor/Slate/Icons/Landscape/Landscape_Erase_x40.png")), Icon20x20));
 
 		//
 		// Icons for selection buttons in PolyEd and TriEd
