Modeling Mode: Add remaining MeshElementSelection related commands to the new viewport toolbar

diff --git a/Engine/Plugins/Editor/ModelingToolsEditorMode/Source/ModelingToolsEditorMode/Private/ModelingModeToolkit_Toolbars.cpp b/Engine/Plugins/Editor/ModelingToolsEditorMode/Source/ModelingToolsEditorMode/Private/ModelingModeToolkit_Toolbars.cpp
--- a/Engine/Plugins/Editor/ModelingToolsEditorMode/Source/ModelingToolsEditorMode/Private/ModelingModeToolkit_Toolbars.cpp
+++ b/Engine/Plugins/Editor/ModelingToolsEditorMode/Source/ModelingToolsEditorMode/Private/ModelingModeToolkit_Toolbars.cpp
@@ -354,6 +354,15 @@ void MakeSubMenu_Selection_DragMode(FModelingToolsEditorModeToolkit* Toolkit, FM
 				return SelectionInteraction->GetActiveDragMode() == DragMode;
 			}));
 	};
+	
+	const FModelingToolsManagerCommands& Commands = FModelingToolsManagerCommands::Get();
+	const TSharedRef<FUICommandList>& CommandList = Toolkit->GetToolkitCommands();
+
+	// map drag mode commands
+	CommandList->MapAction(Commands.SelectionDragMode_None,
+		MakeDragModeOptionAction(EModelingSelectionInteraction_DragMode::NoDragInteraction));
+	CommandList->MapAction(Commands.SelectionDragMode_Path,
+		MakeDragModeOptionAction(EModelingSelectionInteraction_DragMode::PathInteraction));
 
 	MenuBuilder.AddMenuEntry(
 		LOCTEXT("Selection_DragInput_None", "None"), 
@@ -403,6 +412,13 @@ void MakeSubMenu_Selection_MeshType(FModelingToolsEditorModeToolkit* Toolkit, FM
 		{
 			return ModelingMode->bEnableStaticMeshElementSelection;
 		}));
+	
+	const FModelingToolsManagerCommands& Commands = FModelingToolsManagerCommands::Get();
+	const TSharedRef<FUICommandList>& CommandList = Toolkit->GetToolkitCommands();
+
+	// map selection mesh type commands
+	CommandList->MapAction(Commands.SelectionMeshTypes_Volumes, ToggleVolumesAction);
+	CommandList->MapAction(Commands.SelectionMeshTypes_StaticMeshes, ToggleStaticMeshesAction);
 
 
 	MenuBuilder.AddMenuEntry(
@@ -438,6 +454,16 @@ void MakeSubMenu_Selection_LocalFrameMode(FModelingToolsEditorModeToolkit* Toolk
 			}));
 	};
 	
+	const FModelingToolsManagerCommands& Commands = FModelingToolsManagerCommands::Get();
+	const TSharedRef<FUICommandList>& CommandList = Toolkit->GetToolkitCommands();
+
+	// map local frame mode commands
+	CommandList->MapAction(Commands.SelectionLocalFrameMode_Geometry,
+		ToggleFromGeometryAction(EModelingSelectionInteraction_LocalFrameMode::FromGeometry));
+	CommandList->MapAction(Commands.SelectionLocalFrameMode_Object,
+		ToggleFromGeometryAction( EModelingSelectionInteraction_LocalFrameMode::FromObject));
+
+	
 	MenuBuilder.AddMenuEntry(
 		LOCTEXT("Selection_LocalFrameMode_FromGeometry", "From Geometry"), 
 		LOCTEXT("Selection_LocalFrameMode_FromGeometry_Tooltip", "Gizmo Orientation Based on Selected Geometry"),
@@ -745,6 +771,29 @@ void FModelingToolsEditorModeToolkit::MakeSelectionPaletteOverlayWidget()
 			})
 	);
 
+	FUIAction ToggleSelectionLocking(
+		FExecuteAction::CreateLambda([this]
+		{
+			GetMeshSelectionManager()->SetCurrentTargetsLockState(!GetMeshSelectionManager()->GetAnyCurrentTargetsLocked());
+		}),
+		FCanExecuteAction::CreateLambda([this]
+		{
+			return IsInActiveTool() == false
+			&& (GetMeshSelectionManager()->GetMeshTopologyMode() != UGeometrySelectionManager::EMeshTopologyMode::None
+			|| GetMeshSelectionManager()->GetAnyCurrentTargetsLockable() == false);
+		}),
+		FIsActionChecked::CreateLambda([this]()
+		{
+			return GetMeshSelectionManager()->GetAnyCurrentTargetsLocked();
+		})
+	);
+
+	const TSharedRef<FUICommandList>& CommandList = GetToolkitCommands();
+
+	// map lock target commands
+	CommandList->MapAction(Commands.SelectionLocking, ToggleSelectionLocking);
+
+
 	ToolbarBuilder.BeginStyleOverride(FName("SelectionToolBar.RedButton"));
 	ToolbarBuilder.AddToolBarButton(UnlockTargetAction, NAME_None,
 		LOCTEXT("Selection_UnlockTarget", "Unlock"),
diff --git a/Engine/Plugins/Editor/ModelingToolsEditorMode/Source/ModelingToolsEditorMode/Private/ModelingToolsEditorMode.cpp b/Engine/Plugins/Editor/ModelingToolsEditorMode/Source/ModelingToolsEditorMode/Private/ModelingToolsEditorMode.cpp
--- a/Engine/Plugins/Editor/ModelingToolsEditorMode/Source/ModelingToolsEditorMode/Private/ModelingToolsEditorMode.cpp
+++ b/Engine/Plugins/Editor/ModelingToolsEditorMode/Source/ModelingToolsEditorMode/Private/ModelingToolsEditorMode.cpp
@@ -1329,6 +1329,52 @@ void UModelingToolsEditorMode::OnBlueprintPreCompile(UBlueprint* Blueprint)
 					FToolMenuEntry::InitMenuEntryWithCommandList(
 						FModelingToolsManagerCommands::Get().BeginSelectionAction_Contract, CommandList));
 				}
+
+				// Section for Drag Mode setting
+				{
+					FToolMenuSection& DragModeSection =
+						Submenu->FindOrAddSection("Drag Mode", LOCTEXT("DragModeLabel", "Drag Mode"));
+					DragModeSection.AddEntry(
+						FToolMenuEntry::InitMenuEntryWithCommandList(
+							FModelingToolsManagerCommands::Get().SelectionDragMode_None, CommandList));
+					DragModeSection.AddEntry(
+						FToolMenuEntry::InitMenuEntryWithCommandList(
+							FModelingToolsManagerCommands::Get().SelectionDragMode_Path, CommandList));
+				}
+
+				// Section for Mesh Types
+				{
+					FToolMenuSection& MeshTypesSection =
+						Submenu->FindOrAddSection("Selectable Mesh Types", LOCTEXT("SelectableMeshTypesLabel", "Selectable Mesh Types"));
+					MeshTypesSection.AddEntry(
+						FToolMenuEntry::InitMenuEntryWithCommandList(
+							FModelingToolsManagerCommands::Get().SelectionMeshTypes_Volumes, CommandList));
+					MeshTypesSection.AddEntry(
+						FToolMenuEntry::InitMenuEntryWithCommandList(
+							FModelingToolsManagerCommands::Get().SelectionMeshTypes_StaticMeshes, CommandList));
+				}
+				
+				// Section for Local Frame Mode setting
+				{
+					FToolMenuSection& LocalFrameModeSection =
+						Submenu->FindOrAddSection("Local Frame Mode", LOCTEXT("LocalFrameModeLabel", "Local Frame Mode"));
+					LocalFrameModeSection.AddEntry(
+						FToolMenuEntry::InitMenuEntryWithCommandList(
+							FModelingToolsManagerCommands::Get().SelectionLocalFrameMode_Geometry, CommandList));
+					LocalFrameModeSection.AddEntry(
+						FToolMenuEntry::InitMenuEntryWithCommandList(
+							FModelingToolsManagerCommands::Get().SelectionLocalFrameMode_Object, CommandList));
+				}
+
+				// Section for locking/unlocking target
+				{
+					FToolMenuSection& LockingSection =
+						Submenu->FindOrAddSection("Locking Section", FText());
+					LockingSection.AddSeparator(NAME_None);
+					LockingSection.AddEntry(
+					FToolMenuEntry::InitMenuEntryWithCommandList(
+						FModelingToolsManagerCommands::Get().SelectionLocking, CommandList));
+				}
 			}
 		})
 	);
diff --git a/Engine/Plugins/Editor/ModelingToolsEditorMode/Source/ModelingToolsEditorMode/Private/ModelingToolsManagerActions.cpp b/Engine/Plugins/Editor/ModelingToolsEditorMode/Source/ModelingToolsEditorMode/Private/ModelingToolsManagerActions.cpp
--- a/Engine/Plugins/Editor/ModelingToolsEditorMode/Source/ModelingToolsEditorMode/Private/ModelingToolsManagerActions.cpp
+++ b/Engine/Plugins/Editor/ModelingToolsEditorMode/Source/ModelingToolsEditorMode/Private/ModelingToolsManagerActions.cpp
@@ -255,8 +255,15 @@ void FModelingToolsManagerCommands::RegisterCommands()
 	REGISTER_MODELING_ACTION_COMMAND(BeginSelectionAction_Expand, "Expand Selection", "Expand the current Selection by a ring of elements");
 	REGISTER_MODELING_ACTION_COMMAND(BeginSelectionAction_Contract, "Contract Selection", "Contract the current Selection by a ring of elements");
 
-
-
+	REGISTER_MODELING_TOOL_COMMAND_RADIO(SelectionDragMode_None, "None", "No drag input");
+	REGISTER_MODELING_TOOL_COMMAND_RADIO(SelectionDragMode_Path, "Path", "Path drag input");
+	REGISTER_MODELING_TOOL_COMMAND_RADIO(SelectionLocalFrameMode_Geometry, "From Geometry", "Gizmo orientation based on selected geometry");
+	REGISTER_MODELING_TOOL_COMMAND_RADIO(SelectionLocalFrameMode_Object, "From Object", "Gizmo orientation based on object");
+	REGISTER_MODELING_TOOL_COMMAND(SelectionMeshTypes_Volumes, "Volumes", "Toggle whether Volume mesh elements can be selected in the Viewport");
+	REGISTER_MODELING_TOOL_COMMAND(SelectionMeshTypes_StaticMeshes, "Static Meshes", "Toggle whether Static Mesh mesh elements can be selected in the Viewport");
+
+	REGISTER_MODELING_TOOL_COMMAND(SelectionLocking, "Lock Target", "Click to Toggle locking on the Selected Object and disallow/allow Mesh Selections");
+	
 	UI_COMMAND(AddToFavorites, "Add to Favorites", "Add to Favorites", EUserInterfaceActionType::Button, FInputChord());
 	UI_COMMAND(RemoveFromFavorites, "Remove from Favorites", "Remove from Favorites", EUserInterfaceActionType::Button, FInputChord());
 	UI_COMMAND(LoadFavoritesTools, "Faves", "Favorites", EUserInterfaceActionType::RadioButton, FInputChord());
diff --git a/Engine/Plugins/Editor/ModelingToolsEditorMode/Source/ModelingToolsEditorMode/Public/ModelingToolsManagerActions.h b/Engine/Plugins/Editor/ModelingToolsEditorMode/Source/ModelingToolsEditorMode/Public/ModelingToolsManagerActions.h
--- a/Engine/Plugins/Editor/ModelingToolsEditorMode/Source/ModelingToolsEditorMode/Public/ModelingToolsManagerActions.h
+++ b/Engine/Plugins/Editor/ModelingToolsEditorMode/Source/ModelingToolsEditorMode/Public/ModelingToolsManagerActions.h
@@ -193,6 +193,14 @@ public:
 	TSharedPtr<FUICommandInfo> BeginSelectionAction_Extrude;
 	TSharedPtr<FUICommandInfo> BeginSelectionAction_Offset;
 
+	TSharedPtr<FUICommandInfo> SelectionDragMode_None;
+	TSharedPtr<FUICommandInfo> SelectionDragMode_Path;
+	TSharedPtr<FUICommandInfo> SelectionLocalFrameMode_Geometry;
+	TSharedPtr<FUICommandInfo> SelectionLocalFrameMode_Object;
+	TSharedPtr<FUICommandInfo> SelectionMeshTypes_Volumes;
+	TSharedPtr<FUICommandInfo> SelectionMeshTypes_StaticMeshes;
+	TSharedPtr<FUICommandInfo> SelectionLocking;
+
 	TSharedPtr<FUICommandInfo> AddToFavorites;
 	TSharedPtr<FUICommandInfo> RemoveFromFavorites;
 	TSharedPtr<FUICommandInfo> LoadFavoritesTools;
