ModelingTools: Register StylusInputHandlers using the ToolkitHost rather than the ToolkitWidget.

diff --git a/Engine/Plugins/Editor/ModelingToolsEditorMode/Source/ModelingToolsEditorMode/Private/ModelingStylusInputHandler.h b/Engine/Plugins/Editor/ModelingToolsEditorMode/Source/ModelingToolsEditorMode/Private/ModelingStylusInputHandler.h
--- a/Engine/Plugins/Editor/ModelingToolsEditorMode/Source/ModelingToolsEditorMode/Private/ModelingStylusInputHandler.h
+++ b/Engine/Plugins/Editor/ModelingToolsEditorMode/Source/ModelingToolsEditorMode/Private/ModelingStylusInputHandler.h
@@ -9,7 +9,6 @@
 #include "TickableEditorObject.h"
 
 #include "BaseTools/MeshSurfacePointTool.h"
-#include "Containers/SpscQueue.h"
 #include "Widgets/SWidget.h"
 
 namespace UE::Modeling
diff --git a/Engine/Plugins/Editor/ModelingToolsEditorMode/Source/ModelingToolsEditorMode/Private/ModelingToolsEditorModeToolkit.cpp b/Engine/Plugins/Editor/ModelingToolsEditorMode/Source/ModelingToolsEditorMode/Private/ModelingToolsEditorModeToolkit.cpp
--- a/Engine/Plugins/Editor/ModelingToolsEditorMode/Source/ModelingToolsEditorMode/Private/ModelingToolsEditorModeToolkit.cpp
+++ b/Engine/Plugins/Editor/ModelingToolsEditorMode/Source/ModelingToolsEditorMode/Private/ModelingToolsEditorModeToolkit.cpp
@@ -663,6 +663,13 @@ void FModelingToolsEditorModeToolkit::Init(const TSharedPtr<IToolkitHost>& InitT
 
 	UModelingToolsModeCustomizationSettings* ModelingCustomizationSettings = GetMutableDefault<UModelingToolsModeCustomizationSettings>();
 	ModelingCustomizationSettings->OnSettingChanged().AddSP(SharedThis(this),  &FModelingToolsEditorModeToolkit::UpdateSelectionColors);
+
+#if ENABLE_STYLUS_SUPPORT
+	if (InitToolkitHost)
+	{
+		StylusInputHandler->RegisterWindow(InitToolkitHost->GetParentWidget());
+	}
+#endif
 }
 
 void FModelingToolsEditorModeToolkit::MakeToolShutdownOverlayWidget()
@@ -2074,11 +2081,6 @@ void FModelingToolsEditorModeToolkit::InvokeUI()
 	// of FModeToolkit::InlineContentHolder so we can just replace it here
 	InlineContentHolder->SetContent(GetInlineContent().ToSharedRef());
 
-#if ENABLE_STYLUS_SUPPORT
-	// The ToolkitWidget is only attached to a valid window from this point onwards.
-	StylusInputHandler->RegisterWindow(ToolkitWidget.ToSharedRef());
-#endif
-
 	// if the ToolkitBuilder is being used, it will make up the UI
 	if (HasToolkitBuilder())
 	{
