ModelingTools: Added stylus support to ModelingModeToolExtensions.

diff --git a/Engine/Plugins/Editor/ModelingToolsEditorMode/Source/ModelingToolsEditorMode/Private/ModelingToolsEditorMode.cpp b/Engine/Plugins/Editor/ModelingToolsEditorMode/Source/ModelingToolsEditorMode/Private/ModelingToolsEditorMode.cpp
--- a/Engine/Plugins/Editor/ModelingToolsEditorMode/Source/ModelingToolsEditorMode/Private/ModelingToolsEditorMode.cpp
+++ b/Engine/Plugins/Editor/ModelingToolsEditorMode/Source/ModelingToolsEditorMode/Private/ModelingToolsEditorMode.cpp
@@ -861,6 +861,7 @@ void UModelingToolsEditorMode::Enter()
 		FExtensionToolQueryInfo ExtensionQueryInfo;
 		ExtensionQueryInfo.ToolsContext = GetInteractiveToolsContext();
 		ExtensionQueryInfo.AssetAPI = nullptr;
+		ExtensionQueryInfo.StylusAPI = StylusAPI;
 
 		UE_LOG(LogTemp, Log, TEXT("ModelingMode: Found %d Tool Extension Modules"), Extensions.Num());
 		for (int32 k = 0; k < Extensions.Num(); ++k)
diff --git a/Engine/Plugins/Editor/ModelingToolsEditorMode/Source/ModelingToolsEditorMode/Public/ModelingModeToolExtensions.h b/Engine/Plugins/Editor/ModelingToolsEditorMode/Source/ModelingToolsEditorMode/Public/ModelingModeToolExtensions.h
--- a/Engine/Plugins/Editor/ModelingToolsEditorMode/Source/ModelingToolsEditorMode/Public/ModelingModeToolExtensions.h
+++ b/Engine/Plugins/Editor/ModelingToolsEditorMode/Source/ModelingToolsEditorMode/Public/ModelingModeToolExtensions.h
@@ -9,6 +9,7 @@
 class UInteractiveToolsContext;
 class FUICommandInfo;
 class UInteractiveToolBuilder;
+class IToolStylusStateProviderAPI;
 namespace UE
 {
 	class IInteractiveToolCommandsInterface;
@@ -35,6 +36,7 @@ struct MODELINGTOOLSEDITORMODE_API FExtensionToolQueryInfo
 
 	UInteractiveToolsContext* ToolsContext = nullptr;
 	FModelingModeAssetAPI* AssetAPI = nullptr;
+	IToolStylusStateProviderAPI* StylusAPI = nullptr;
 };
 
 /**
