Modeling Mode: fix bug that produced an error in the output log when opening the Mesh Element Selection toolbar menu

diff --git a/Engine/Plugins/Editor/ModelingToolsEditorMode/Source/ModelingToolsEditorMode/Private/ModelingModeToolkit_Toolbars.cpp b/Engine/Plugins/Editor/ModelingToolsEditorMode/Source/ModelingToolsEditorMode/Private/ModelingModeToolkit_Toolbars.cpp
--- a/Engine/Plugins/Editor/ModelingToolsEditorMode/Source/ModelingToolsEditorMode/Private/ModelingModeToolkit_Toolbars.cpp
+++ b/Engine/Plugins/Editor/ModelingToolsEditorMode/Source/ModelingToolsEditorMode/Private/ModelingModeToolkit_Toolbars.cpp
@@ -334,46 +334,27 @@ void MakeSubMenu_DefaultMeshObjectType(FMenuBuilder& MenuBuilder)
 	MakeSubMenu_DefaultMeshObjectPhysicsSettings(MenuBuilder);
 }
 
-void MakeSubMenu_Selection_DragMode(FModelingToolsEditorModeToolkit* Toolkit, FMenuBuilder& MenuBuilder)
+void MakeSubMenu_Selection_DragMode(FMenuBuilder& MenuBuilder)
 {
-	UModelingSelectionInteraction* SelectionInteraction = Cast<UModelingToolsEditorMode>(Toolkit->GetScriptableEditorMode())->GetSelectionInteraction();
-
-	auto MakeDragModeOptionAction = [SelectionInteraction](EModelingSelectionInteraction_DragMode DragMode)
-	{
-		return FUIAction(
-			FExecuteAction::CreateLambda([SelectionInteraction, DragMode]
-			{
-				SelectionInteraction->SetActiveDragMode(DragMode);
-				UModelingToolsModeCustomizationSettings* ModelingEditorSettings = GetMutableDefault<UModelingToolsModeCustomizationSettings>();
-				ModelingEditorSettings->LastMeshSelectionDragMode = static_cast<int>(DragMode);
-				ModelingEditorSettings->SaveConfig();
-			}),
-			FCanExecuteAction(),
-			FIsActionChecked::CreateLambda([SelectionInteraction, DragMode]()
-			{
-				return SelectionInteraction->GetActiveDragMode() == DragMode;
-			}));
-	};
-	
 	const FModelingToolsManagerCommands& Commands = FModelingToolsManagerCommands::Get();
-	const TSharedRef<FUICommandList>& CommandList = Toolkit->GetToolkitCommands();
-
-	// map drag mode commands
-	CommandList->MapAction(Commands.SelectionDragMode_None,
-		MakeDragModeOptionAction(EModelingSelectionInteraction_DragMode::NoDragInteraction));
-	CommandList->MapAction(Commands.SelectionDragMode_Path,
-		MakeDragModeOptionAction(EModelingSelectionInteraction_DragMode::PathInteraction));
 
 	MenuBuilder.AddMenuEntry(
-		LOCTEXT("Selection_DragInput_None", "None"), 
+		Commands.SelectionDragMode_None,
+		NAME_None,
+		LOCTEXT("Selection_DragInput_None", "None"),
 		LOCTEXT("Selection_DragInput_None_Tooltip", "No drag input"),
-		FSlateIcon(), MakeDragModeOptionAction(EModelingSelectionInteraction_DragMode::NoDragInteraction), NAME_None, EUserInterfaceActionType::ToggleButton);
-
+		FSlateIcon(),
+		NAME_None
+	);
 	MenuBuilder.AddMenuEntry(
-		LOCTEXT("Selection_DragInput_Path", "Path"), 
+		Commands.SelectionDragMode_Path,
+		NAME_None,
+		LOCTEXT("Selection_DragInput_Path", "Path"),
 		LOCTEXT("Selection_DragInput_Path_Tooltip", "Path Drag Input"),
-		FSlateIcon(), MakeDragModeOptionAction(EModelingSelectionInteraction_DragMode::PathInteraction), NAME_None, EUserInterfaceActionType::ToggleButton);
-
+		FSlateIcon(),
+		NAME_None
+	);
+	
 	// marquee mode is not functional yet, so this is disabled
 	//MenuBuilder.AddMenuEntry(
 	//	LOCTEXT("Selection_DragInput_Marquee", "Rectangle"), 
@@ -383,112 +364,63 @@ void MakeSubMenu_Selection_DragMode(FModelingToolsEditorModeToolkit* Toolkit, FM
 }
 
 
-void MakeSubMenu_Selection_MeshType(FModelingToolsEditorModeToolkit* Toolkit, FMenuBuilder& MenuBuilder)
+void MakeSubMenu_Selection_MeshType(FMenuBuilder& MenuBuilder)
 {
-	UModelingToolsEditorMode* ModelingMode = Cast<UModelingToolsEditorMode>(Toolkit->GetScriptableEditorMode());
-	FUIAction ToggleVolumesAction(
-		FExecuteAction::CreateLambda([ModelingMode]
-		{
-			ModelingMode->bEnableVolumeElementSelection = ! ModelingMode->bEnableVolumeElementSelection;
-
-			UModelingToolsModeCustomizationSettings* ModelingEditorSettings = GetMutableDefault<UModelingToolsModeCustomizationSettings>();
-			ModelingEditorSettings->bLastMeshSelectionVolumeToggle = ModelingMode->bEnableVolumeElementSelection;
-		}),
-		FCanExecuteAction(),
-		FIsActionChecked::CreateLambda([ModelingMode]()
-		{
-			return ModelingMode->bEnableVolumeElementSelection;
-		}));
-	FUIAction ToggleStaticMeshesAction(
-		FExecuteAction::CreateLambda([ModelingMode]
-		{
-			ModelingMode->bEnableStaticMeshElementSelection = ! ModelingMode->bEnableStaticMeshElementSelection;
-
-			UModelingToolsModeCustomizationSettings* ModelingEditorSettings = GetMutableDefault<UModelingToolsModeCustomizationSettings>();
-			ModelingEditorSettings->bLastMeshSelectionStaticMeshToggle = ModelingMode->bEnableStaticMeshElementSelection;
-		}),
-		FCanExecuteAction(),
-		FIsActionChecked::CreateLambda([ModelingMode]()
-		{
-			return ModelingMode->bEnableStaticMeshElementSelection;
-		}));
-	
 	const FModelingToolsManagerCommands& Commands = FModelingToolsManagerCommands::Get();
-	const TSharedRef<FUICommandList>& CommandList = Toolkit->GetToolkitCommands();
-
-	// map selection mesh type commands
-	CommandList->MapAction(Commands.SelectionMeshTypes_Volumes, ToggleVolumesAction);
-	CommandList->MapAction(Commands.SelectionMeshTypes_StaticMeshes, ToggleStaticMeshesAction);
-
 
 	MenuBuilder.AddMenuEntry(
+		Commands.SelectionMeshTypes_Volumes,
+		NAME_None,
 		LOCTEXT("Selection_MeshTypes_Volumes", "Volumes"), 
 		LOCTEXT("Selection_MeshTypes_Volumes_Tooltip", "Toggle whether Volume mesh elements can be selected in the Viewport"),
-		FSlateIcon(), ToggleVolumesAction, NAME_None, EUserInterfaceActionType::ToggleButton);
-
+		FSlateIcon(),
+		NAME_None
+	);
 	MenuBuilder.AddMenuEntry(
+		Commands.SelectionMeshTypes_StaticMeshes,
+		NAME_None,
 		LOCTEXT("Selection_MeshTypes_StaticMesh", "Static Meshes"), 
 		LOCTEXT("Selection_MeshTypes_StaticMesh_Tooltip", "Toggle whether Static Mesh mesh elements can be selected in the Viewport"),
-		FSlateIcon(), ToggleStaticMeshesAction, NAME_None, EUserInterfaceActionType::ToggleButton);
-
+		FSlateIcon(),
+		NAME_None
+	);
 }
 	
-void MakeSubMenu_Selection_LocalFrameMode(FModelingToolsEditorModeToolkit* Toolkit, FMenuBuilder& MenuBuilder)
+void MakeSubMenu_Selection_LocalFrameMode(FMenuBuilder& MenuBuilder)
 {
-		
-	UModelingSelectionInteraction* SelectionInteraction = Cast<UModelingToolsEditorMode>(Toolkit->GetScriptableEditorMode())->GetSelectionInteraction();
-	auto ToggleFromGeometryAction = [SelectionInteraction](EModelingSelectionInteraction_LocalFrameMode LocalFrameMode)
-	{
-		return FUIAction(
-			FExecuteAction::CreateLambda([SelectionInteraction, LocalFrameMode]
-			{
-				SelectionInteraction->SetLocalFrameMode(LocalFrameMode);
-				UModelingToolsModeCustomizationSettings* ModelingEditorSettings = GetMutableDefault<UModelingToolsModeCustomizationSettings>();
-				ModelingEditorSettings->LastMeshSelectionLocalFrameMode = static_cast<int>(LocalFrameMode);
-				ModelingEditorSettings->SaveConfig();
-			}),
-			FCanExecuteAction(),
-			FIsActionChecked::CreateLambda([SelectionInteraction, LocalFrameMode]()
-			{
-				return (SelectionInteraction->GetLocalFrameMode() == LocalFrameMode);
-			}));
-	};
-	
 	const FModelingToolsManagerCommands& Commands = FModelingToolsManagerCommands::Get();
-	const TSharedRef<FUICommandList>& CommandList = Toolkit->GetToolkitCommands();
-
-	// map local frame mode commands
-	CommandList->MapAction(Commands.SelectionLocalFrameMode_Geometry,
-		ToggleFromGeometryAction(EModelingSelectionInteraction_LocalFrameMode::FromGeometry));
-	CommandList->MapAction(Commands.SelectionLocalFrameMode_Object,
-		ToggleFromGeometryAction( EModelingSelectionInteraction_LocalFrameMode::FromObject));
-
-	
 	MenuBuilder.AddMenuEntry(
+		Commands.SelectionLocalFrameMode_Geometry,
+		NAME_None,
 		LOCTEXT("Selection_LocalFrameMode_FromGeometry", "From Geometry"), 
 		LOCTEXT("Selection_LocalFrameMode_FromGeometry_Tooltip", "Gizmo Orientation Based on Selected Geometry"),
-		FSlateIcon(), ToggleFromGeometryAction(EModelingSelectionInteraction_LocalFrameMode::FromGeometry), NAME_None, EUserInterfaceActionType::ToggleButton);
-
+		FSlateIcon(),
+		NAME_None
+	);
 	MenuBuilder.AddMenuEntry(
+		Commands.SelectionLocalFrameMode_Object,
+		NAME_None,
 		LOCTEXT("Selection_LocalFrameMode_FromObject", "From Object"), 
 		LOCTEXT("Selection_LocalFrameMode_FromObject_Tooltip", "Gizmo Orientation Based on Object"),
-		FSlateIcon(), ToggleFromGeometryAction(EModelingSelectionInteraction_LocalFrameMode::FromObject), NAME_None, EUserInterfaceActionType::ToggleButton);
+		FSlateIcon(),
+		NAME_None
+	);
 }
 
-TSharedRef<SWidget> MakeMenu_SelectionConfigSettings(FModelingToolsEditorModeToolkit* Toolkit)
+TSharedRef<SWidget> MakeMenu_SelectionConfigSettings(const FModelingToolsEditorModeToolkit* Toolkit)
 {
-	FMenuBuilder MenuBuilder(true, TSharedPtr<FUICommandList>());
+	FMenuBuilder MenuBuilder(true, Toolkit->GetToolkitCommands());
 
 	MenuBuilder.BeginSection("Section_DragMode", LOCTEXT("Section_DragMode", "Drag Mode"));
-	MakeSubMenu_Selection_DragMode(Toolkit, MenuBuilder);
+	MakeSubMenu_Selection_DragMode(MenuBuilder);
 	MenuBuilder.EndSection();
 
 	MenuBuilder.BeginSection("Section_MeshTypes", LOCTEXT("Section_MeshTypes", "Selectable Mesh Types"));
-	MakeSubMenu_Selection_MeshType(Toolkit, MenuBuilder);
+	MakeSubMenu_Selection_MeshType(MenuBuilder);
 	MenuBuilder.EndSection();
 
 	MenuBuilder.BeginSection("Section_LocalFrameMode", LOCTEXT("Section_LocalFrameMode", "Local Frame Mode"));
-	MakeSubMenu_Selection_LocalFrameMode(Toolkit, MenuBuilder);
+	MakeSubMenu_Selection_LocalFrameMode(MenuBuilder);
 	MenuBuilder.EndSection();
 
 	TSharedRef<SWidget> MenuWidget = MenuBuilder.MakeWidget();
diff --git a/Engine/Plugins/Editor/ModelingToolsEditorMode/Source/ModelingToolsEditorMode/Private/ModelingToolsEditorMode.cpp b/Engine/Plugins/Editor/ModelingToolsEditorMode/Source/ModelingToolsEditorMode/Private/ModelingToolsEditorMode.cpp
--- a/Engine/Plugins/Editor/ModelingToolsEditorMode/Source/ModelingToolsEditorMode/Private/ModelingToolsEditorMode.cpp
+++ b/Engine/Plugins/Editor/ModelingToolsEditorMode/Source/ModelingToolsEditorMode/Private/ModelingToolsEditorMode.cpp
@@ -825,14 +825,96 @@ void UModelingToolsEditorMode::Enter()
 		RegisterSelectionMode(UGeometrySelectionManager::EMeshTopologyMode::Polygroup, UGeometrySelectionManager::EGeometryElementType::Edge, ToolManagerCommands.MeshSelectionModeAction_GroupEdges);
 	}
 
+	// Register Drag Modes
+	
+	auto MakeDragModeOptionAction = [this](EModelingSelectionInteraction_DragMode DragMode)
+	{
+		return FUIAction(
+			FExecuteAction::CreateLambda([this, DragMode]
+			{
+				SelectionInteraction->SetActiveDragMode(DragMode);
+				UModelingToolsModeCustomizationSettings* ModelingEditorSettings = GetMutableDefault<UModelingToolsModeCustomizationSettings>();
+				ModelingEditorSettings->LastMeshSelectionDragMode = static_cast<int>(DragMode);
+				ModelingEditorSettings->SaveConfig();
+			}),
+			FCanExecuteAction(),
+			FIsActionChecked::CreateLambda([this, DragMode]()
+			{
+				return SelectionInteraction->GetActiveDragMode() == DragMode;
+			}));
+	};
+	
+	const TSharedRef<FUICommandList>& CommandList = Toolkit->GetToolkitCommands();
+
+	// map drag mode commands
+	CommandList->MapAction(ToolManagerCommands.SelectionDragMode_None,
+		MakeDragModeOptionAction(EModelingSelectionInteraction_DragMode::NoDragInteraction));
+	CommandList->MapAction(ToolManagerCommands.SelectionDragMode_Path,
+		MakeDragModeOptionAction(EModelingSelectionInteraction_DragMode::PathInteraction));
+
+	// Register Selection Mesh Types
+
+	FUIAction ToggleVolumesAction(
+		FExecuteAction::CreateLambda([this]
+		{
+			bEnableVolumeElementSelection = !bEnableVolumeElementSelection;
+
+			UModelingToolsModeCustomizationSettings* ModelingEditorSettings = GetMutableDefault<UModelingToolsModeCustomizationSettings>();
+			ModelingEditorSettings->bLastMeshSelectionVolumeToggle = bEnableVolumeElementSelection;
+		}),
+		FCanExecuteAction(),
+		FIsActionChecked::CreateLambda([this]()
+		{
+			return bEnableVolumeElementSelection;
+		}));
+	FUIAction ToggleStaticMeshesAction(
+		FExecuteAction::CreateLambda([this]
+		{	bEnableStaticMeshElementSelection = !bEnableStaticMeshElementSelection;
+
+			UModelingToolsModeCustomizationSettings* ModelingEditorSettings = GetMutableDefault<UModelingToolsModeCustomizationSettings>();
+			ModelingEditorSettings->bLastMeshSelectionStaticMeshToggle = bEnableStaticMeshElementSelection;
+		}),
+		FCanExecuteAction(),
+		FIsActionChecked::CreateLambda([this]()
+		{
+			return bEnableStaticMeshElementSelection;
+		}));
+	
+	// map selection mesh type commands
+	CommandList->MapAction(ToolManagerCommands.SelectionMeshTypes_Volumes, ToggleVolumesAction);
+	CommandList->MapAction(ToolManagerCommands.SelectionMeshTypes_StaticMeshes, ToggleStaticMeshesAction);
+
+	// Register Frame Modes
+
+	auto ToggleFromGeometryAction = [this](EModelingSelectionInteraction_LocalFrameMode LocalFrameMode)
+	{
+		return FUIAction(
+			FExecuteAction::CreateLambda([this, LocalFrameMode]
+			{
+				SelectionInteraction->SetLocalFrameMode(LocalFrameMode);
+				UModelingToolsModeCustomizationSettings* ModelingEditorSettings = GetMutableDefault<UModelingToolsModeCustomizationSettings>();
+				ModelingEditorSettings->LastMeshSelectionLocalFrameMode = static_cast<int>(LocalFrameMode);
+				ModelingEditorSettings->SaveConfig();
+			}),
+			FCanExecuteAction(),
+			FIsActionChecked::CreateLambda([this, LocalFrameMode]()
+			{
+				return (SelectionInteraction->GetLocalFrameMode() == LocalFrameMode);
+			}));
+	};
+	
+	CommandList->MapAction(ToolManagerCommands.SelectionLocalFrameMode_Geometry,
+		ToggleFromGeometryAction(EModelingSelectionInteraction_LocalFrameMode::FromGeometry));
+	CommandList->MapAction(ToolManagerCommands.SelectionLocalFrameMode_Object,
+		ToggleFromGeometryAction( EModelingSelectionInteraction_LocalFrameMode::FromObject));
+
 	// Colors initialized here any time Modeling mode is entered
 	SelectionManager->SetSelectionColors(ModelingEditorSettings->UnselectedColor, ModelingEditorSettings->HoverOverSelectedColor, ModelingEditorSettings->HoverOverUnselectedColor, ModelingEditorSettings->GeometrySelectedColor);
 
 	// this function registers and tracks an active UGeometrySelectionEditCommand and it's associated UICommand
-	auto RegisterSelectionCommand = [&](UGeometrySelectionEditCommand* Command, TSharedPtr<FUICommandInfo> UICommand, bool bAlwaysVisible)
+	auto RegisterSelectionCommand = [this, &CommandList](UGeometrySelectionEditCommand* Command, TSharedPtr<FUICommandInfo> UICommand, bool bAlwaysVisible)
 	{
 		ModelingModeCommands.Add(Command);
-		const TSharedRef<FUICommandList>& CommandList = Toolkit->GetToolkitCommands();
 		CommandList->MapAction(UICommand,
 			FExecuteAction::CreateUObject(GetSelectionManager(), &UGeometrySelectionManager::ExecuteSelectionCommand, Command),
 			FCanExecuteAction::CreateUObject(GetSelectionManager(), &UGeometrySelectionManager::CanExecuteSelectionCommand, Command),
