Modeling Mode: allow enabling/disabling new brush selection in VSculpt Tool

diff --git a/Engine/Plugins/Editor/ModelingToolsEditorMode/Source/ModelingToolsEditorMode/Private/DetailsCustomizations/MeshVertexSculptToolCustomizations.cpp b/Engine/Plugins/Editor/ModelingToolsEditorMode/Source/ModelingToolsEditorMode/Private/DetailsCustomizations/MeshVertexSculptToolCustomizations.cpp
--- a/Engine/Plugins/Editor/ModelingToolsEditorMode/Source/ModelingToolsEditorMode/Private/DetailsCustomizations/MeshVertexSculptToolCustomizations.cpp
+++ b/Engine/Plugins/Editor/ModelingToolsEditorMode/Source/ModelingToolsEditorMode/Private/DetailsCustomizations/MeshVertexSculptToolCustomizations.cpp
@@ -160,8 +160,16 @@ void FVertexBrushSculptPropertiesDetails::CustomizeDetails(IDetailLayoutBuilder&
 		.FlyoutTileSize(FVector2D(FlyoutIconSize, FlyoutIconSize))
 		.FlyoutSize(FVector2D(FlyoutWidth, 1))
 		.ListItems(BrushTypeItems)
+		.IsEnabled_Lambda([this]()
+		{
+			if (const UMeshVertexSculptTool* Tool = TargetTool.Get())
+			{
+				return Tool->CanUpdateBrushType();
+			}
+			return false;
+		})
 		.OnSelectionChanged_Lambda([this](TSharedPtr<SComboPanel::FComboPanelItem> NewSelectedItem) {
-			if ( TargetTool.IsValid() )
+			if ( UMeshVertexSculptTool* Tool = TargetTool.Get() )
 			{
 				TargetTool.Get()->SetActiveBrushType(NewSelectedItem->Identifier);
 			}
diff --git a/Engine/Plugins/Experimental/MeshModelingToolsetExp/Source/MeshModelingToolsExp/Private/MeshVertexSculptTool.cpp b/Engine/Plugins/Experimental/MeshModelingToolsetExp/Source/MeshModelingToolsExp/Private/MeshVertexSculptTool.cpp
--- a/Engine/Plugins/Experimental/MeshModelingToolsetExp/Source/MeshModelingToolsExp/Private/MeshVertexSculptTool.cpp
+++ b/Engine/Plugins/Experimental/MeshModelingToolsetExp/Source/MeshModelingToolsExp/Private/MeshVertexSculptTool.cpp
@@ -980,6 +980,11 @@ void UMeshVertexSculptTool::RealignBrush(FMeshSculptBrushOp::EStampAlignmentType
 	}
 }
 
+bool UMeshVertexSculptTool::CanUpdateBrushType() const
+{
+	return DefaultPrimaryBrushID == -1;
+}
+
 TFuture<void> UMeshVertexSculptTool::ApplyStamp()
 {
 	TRACE_CPUPROFILER_EVENT_SCOPE(VtxSculptTool_ApplyStamp);
diff --git a/Engine/Plugins/Experimental/MeshModelingToolsetExp/Source/MeshModelingToolsExp/Public/MeshVertexSculptTool.h b/Engine/Plugins/Experimental/MeshModelingToolsetExp/Source/MeshModelingToolsExp/Public/MeshVertexSculptTool.h
--- a/Engine/Plugins/Experimental/MeshModelingToolsetExp/Source/MeshModelingToolsExp/Public/MeshVertexSculptTool.h
+++ b/Engine/Plugins/Experimental/MeshModelingToolsetExp/Source/MeshModelingToolsExp/Public/MeshVertexSculptTool.h
@@ -260,6 +260,8 @@ public:
 
 	void SetDefaultPrimaryBrushID(int32 InPrimaryBrushID);
 
+	bool CanUpdateBrushType() const;
+
 	/**
 	* OnDetailsPanelRequestRebuild is broadcast by the tool when it detects it needs to have it's details panel rebuilt outside
 	* of normal rebuilding triggers, such as changing property set objects. This is useful in rare circumstances, such as when
