ModelingTools: Move SelectionLocking command binding out from MakeSelectionPaletteOverlay out to the mode.

diff --git a/Engine/Plugins/Editor/ModelingToolsEditorMode/Source/ModelingToolsEditorMode/Private/ModelingModeToolkit_Toolbars.cpp b/Engine/Plugins/Editor/ModelingToolsEditorMode/Source/ModelingToolsEditorMode/Private/ModelingModeToolkit_Toolbars.cpp
--- a/Engine/Plugins/Editor/ModelingToolsEditorMode/Source/ModelingToolsEditorMode/Private/ModelingModeToolkit_Toolbars.cpp
+++ b/Engine/Plugins/Editor/ModelingToolsEditorMode/Source/ModelingToolsEditorMode/Private/ModelingModeToolkit_Toolbars.cpp
@@ -703,29 +703,6 @@ void FModelingToolsEditorModeToolkit::MakeSelectionPaletteOverlayWidget()
 			})
 	);
 
-	FUIAction ToggleSelectionLocking(
-		FExecuteAction::CreateLambda([this]
-		{
-			GetMeshSelectionManager()->SetCurrentTargetsLockState(!GetMeshSelectionManager()->GetAnyCurrentTargetsLocked());
-		}),
-		FCanExecuteAction::CreateLambda([this]
-		{
-			return IsInActiveTool() == false
-			&& (GetMeshSelectionManager()->GetMeshTopologyMode() != UGeometrySelectionManager::EMeshTopologyMode::None
-			|| GetMeshSelectionManager()->GetAnyCurrentTargetsLockable() == false);
-		}),
-		FIsActionChecked::CreateLambda([this]()
-		{
-			return GetMeshSelectionManager()->GetAnyCurrentTargetsLocked();
-		})
-	);
-
-	const TSharedRef<FUICommandList>& CommandList = GetToolkitCommands();
-
-	// map lock target commands
-	CommandList->MapAction(Commands.SelectionLocking, ToggleSelectionLocking);
-
-
 	ToolbarBuilder.BeginStyleOverride(FName("SelectionToolBar.RedButton"));
 	ToolbarBuilder.AddToolBarButton(UnlockTargetAction, NAME_None,
 		LOCTEXT("Selection_UnlockTarget", "Unlock"),
diff --git a/Engine/Plugins/Editor/ModelingToolsEditorMode/Source/ModelingToolsEditorMode/Private/ModelingToolsEditorMode.cpp b/Engine/Plugins/Editor/ModelingToolsEditorMode/Source/ModelingToolsEditorMode/Private/ModelingToolsEditorMode.cpp
--- a/Engine/Plugins/Editor/ModelingToolsEditorMode/Source/ModelingToolsEditorMode/Private/ModelingToolsEditorMode.cpp
+++ b/Engine/Plugins/Editor/ModelingToolsEditorMode/Source/ModelingToolsEditorMode/Private/ModelingToolsEditorMode.cpp
@@ -935,6 +935,26 @@ void UModelingToolsEditorMode::Enter()
 	RegisterSelectionCommand(NewObject<UModifyGeometrySelectionCommand_Expand>(), ToolManagerCommands.BeginSelectionAction_Expand, true);
 	RegisterSelectionCommand(NewObject<UModifyGeometrySelectionCommand_Contract>(), ToolManagerCommands.BeginSelectionAction_Contract, true);
 
+	// Register Selection Lock Toggle
+	FUIAction ToggleSelectionLocking(
+		FExecuteAction::CreateLambda([this]
+		{
+			SelectionManager->SetCurrentTargetsLockState(!SelectionManager->GetAnyCurrentTargetsLocked());
+		}),
+		FCanExecuteAction::CreateLambda([this]
+		{
+			return Toolkit.IsValid()
+			&& StaticCast<FModelingToolsEditorModeToolkit*>(Toolkit.Get())->IsInActiveTool() == false
+			&& SelectionManager->GetMeshTopologyMode() != UGeometrySelectionManager::EMeshTopologyMode::None
+			&& SelectionManager->GetAnyCurrentTargetsLockable();
+		}),
+		FIsActionChecked::CreateLambda([this]()
+		{
+			return SelectionManager->GetAnyCurrentTargetsLocked();
+		})
+	);
+	CommandList->MapAction(ToolManagerCommands.SelectionLocking, ToggleSelectionLocking);
+
 	// register extensions
 	TArray<IModelingModeToolExtension*> Extensions = IModularFeatures::Get().GetModularFeatureImplementations<IModelingModeToolExtension>(
 		IModelingModeToolExtension::GetModularFeatureName());
@@ -1710,9 +1730,10 @@ bool UModelingToolsEditorMode::BoxSelect(FBox& InBox, bool InSelect)
 
 bool UModelingToolsEditorMode::FrustumSelect(const FConvexVolume& InFrustum, FEditorViewportClient* InViewportClient, bool InSelect)
 {
-	if (bIsToolActive)
+	// Disable FrustumSelect when a tool is active
+	if (!Toolkit.IsValid() || StaticCast<FModelingToolsEditorModeToolkit*>(Toolkit.Get())->IsInActiveTool() )
 	{
-		return true; // will signal that Frustum Select does not need to do anything; disables all frustum select when tool is active
+		return true;
 	}
 
 	if (GetMeshElementSelectionSystemEnabled()
@@ -1772,8 +1793,6 @@ void UModelingToolsEditorMode::OnToolStarted(UInteractiveToolManager* Manager, U
 	// the result. This apparently broken behavior is currently by-design.
 	FSlateThrottleManager::Get().DisableThrottle(true);
 
-	bIsToolActive = true;
-
 	FModelingToolActionCommands::UpdateToolCommandBinding(Tool, Toolkit->GetToolkitCommands(), false);
 	if (TryGetExtensionToolCommandGetter(Manager, Tool, ExtensionToolCommandsGetter) && ensure(ExtensionToolCommandsGetter))
 	{
@@ -1802,8 +1821,6 @@ void UModelingToolsEditorMode::OnToolEnded(UInteractiveToolManager* Manager, UIn
 	// re-enable slate throttling (see OnToolStarted)
 	FSlateThrottleManager::Get().DisableThrottle(false);
 
-	bIsToolActive = false;
-
 	FModelingToolActionCommands::UpdateToolCommandBinding(Tool, Toolkit->GetToolkitCommands(), true);
 	if (ExtensionToolCommandsGetter)
 	{
diff --git a/Engine/Plugins/Editor/ModelingToolsEditorMode/Source/ModelingToolsEditorMode/Public/ModelingToolsEditorMode.h b/Engine/Plugins/Editor/ModelingToolsEditorMode/Source/ModelingToolsEditorMode/Public/ModelingToolsEditorMode.h
--- a/Engine/Plugins/Editor/ModelingToolsEditorMode/Source/ModelingToolsEditorMode/Public/ModelingToolsEditorMode.h
+++ b/Engine/Plugins/Editor/ModelingToolsEditorMode/Source/ModelingToolsEditorMode/Public/ModelingToolsEditorMode.h
@@ -173,8 +173,6 @@ protected:
 	// tracking of unlocked stuff
 	static FDelegateHandle GlobalModelingWorldTeardownEventHandle;
 private:
-	bool bIsToolActive = false;
-
 	// add modeling-mode-specific portion of new viewport toolbar
 	static void PopulateModelingModeViewportToolbar(const FName InMenuName, const TSharedPtr<const FUICommandList>& InCommandList);
 	// remove modeling-mode-specific portion of new viewport toolbar
