Modeling Tools: minor updates to default selection colors

diff --git a/Engine/Plugins/Editor/ModelingToolsEditorMode/Source/ModelingToolsEditorMode/Public/ModelingToolsEditorModeSettings.h b/Engine/Plugins/Editor/ModelingToolsEditorMode/Source/ModelingToolsEditorMode/Public/ModelingToolsEditorModeSettings.h
--- a/Engine/Plugins/Editor/ModelingToolsEditorMode/Source/ModelingToolsEditorMode/Public/ModelingToolsEditorModeSettings.h
+++ b/Engine/Plugins/Editor/ModelingToolsEditorMode/Source/ModelingToolsEditorMode/Public/ModelingToolsEditorModeSettings.h
@@ -359,19 +359,19 @@ public:
 	 * Custom color for when geometry is not a part of current selection when using Mesh Element Selection
 	 */
 	UPROPERTY(Config, EditAnywhere, Category = "Modeling Mode|Mesh Element Selection Mode")
-	FLinearColor UnselectedColor = FLinearColor(.068478f, 0.56478f, 0.708376f,.5f); // default red
+	FLinearColor UnselectedColor = FLinearColor(.068478f, 0.56478f, 0.708376f,.5f); // default light blue
 
 	/**
 	 * Custom color hovering over selected geometry in Mesh Element Selection
 	 */
 	UPROPERTY(Config, EditAnywhere, Category = "Modeling Mode|Mesh Element Selection Mode")
-	FLinearColor HoverOverSelectedColor = FLinearColor(1.0f, 0.943014f, 0.74f); // default pink
+	FLinearColor HoverOverSelectedColor = FLinearColor(1.0f, 0.943014f, 0.74f); // default white
 
 	/**
 	* Custom color hovering over unselected geometry in Mesh Element Selection
 	*/
 	UPROPERTY(Config, EditAnywhere, Category = "Modeling Mode|Mesh Element Selection Mode")
-	FLinearColor HoverOverUnselectedColor = FLinearColor(1.0f, 0.25f, 0.712328f, .5f); // default green
+	FLinearColor HoverOverUnselectedColor = FLinearColor(1.0f, 0.25f, 0.712328f, .5f); // default pink
 
 	/**
 	* Custom color for when geometry is a part of the current selection when using Mesh Element Selection
