ModelingTools: Disable StylusAPI when using NullRHI

diff --git a/Engine/Plugins/Editor/ModelingToolsEditorMode/Source/ModelingToolsEditorMode/ModelingToolsEditorMode.Build.cs b/Engine/Plugins/Editor/ModelingToolsEditorMode/Source/ModelingToolsEditorMode/ModelingToolsEditorMode.Build.cs
--- a/Engine/Plugins/Editor/ModelingToolsEditorMode/Source/ModelingToolsEditorMode/ModelingToolsEditorMode.Build.cs
+++ b/Engine/Plugins/Editor/ModelingToolsEditorMode/Source/ModelingToolsEditorMode/ModelingToolsEditorMode.Build.cs
@@ -20,6 +20,7 @@ public class ModelingToolsEditorMode : ModuleRules
 			new string[]
 			{
 				"CoreUObject",
+				"RHI",
 				"Slate",
 				"SlateCore",
 				"Engine",
diff --git a/Engine/Plugins/Editor/ModelingToolsEditorMode/Source/ModelingToolsEditorMode/Private/ModelingToolsEditorModeToolkit.cpp b/Engine/Plugins/Editor/ModelingToolsEditorMode/Source/ModelingToolsEditorMode/Private/ModelingToolsEditorModeToolkit.cpp
--- a/Engine/Plugins/Editor/ModelingToolsEditorMode/Source/ModelingToolsEditorMode/Private/ModelingToolsEditorModeToolkit.cpp
+++ b/Engine/Plugins/Editor/ModelingToolsEditorMode/Source/ModelingToolsEditorMode/Private/ModelingToolsEditorModeToolkit.cpp
@@ -179,7 +179,11 @@ FModelingToolsEditorModeToolkit::FModelingToolsEditorModeToolkit()
 	CurrentPreset = MakeShared<FAssetData>();
 	
 #if ENABLE_STYLUS_SUPPORT
-	StylusInputHandler = MakeUnique<UE::Modeling::FStylusInputHandler>();
+	// Do not initialize StylusInput for headless editor instances.
+	if (!GUsingNullRHI)
+	{
+		StylusInputHandler = MakeUnique<UE::Modeling::FStylusInputHandler>();
+	}
 #endif
 }
 
@@ -666,7 +670,7 @@ void FModelingToolsEditorModeToolkit::Init(const TSharedPtr<IToolkitHost>& InitT
 	ModelingCustomizationSettings->OnSettingChanged().AddSP(SharedThis(this),  &FModelingToolsEditorModeToolkit::UpdateSelectionColors);
 
 #if ENABLE_STYLUS_SUPPORT
-	if (InitToolkitHost)
+	if (StylusInputHandler && InitToolkitHost)
 	{
 		StylusInputHandler->RegisterWindow(InitToolkitHost->GetParentWidget());
 	}
