Modeling Mode: Add support for and the ability to toggle Pressure Sensitivity for VertexSculpt & VertexPaint BrushSize

diff --git a/Engine/Plugins/Editor/ModelingToolsEditorMode/Source/ModelingToolsEditorMode/Private/DetailsCustomizations/ModelingToolsBrushSizeCustomization.cpp b/Engine/Plugins/Editor/ModelingToolsEditorMode/Source/ModelingToolsEditorMode/Private/DetailsCustomizations/ModelingToolsBrushSizeCustomization.cpp
--- a/Engine/Plugins/Editor/ModelingToolsEditorMode/Source/ModelingToolsEditorMode/Private/DetailsCustomizations/ModelingToolsBrushSizeCustomization.cpp
+++ b/Engine/Plugins/Editor/ModelingToolsEditorMode/Source/ModelingToolsEditorMode/Private/DetailsCustomizations/ModelingToolsBrushSizeCustomization.cpp
@@ -3,6 +3,7 @@
 #include "DetailsCustomizations/ModelingToolsBrushSizeCustomization.h"
 
 #include "DetailWidgetRow.h"
+#include "ModelingToolsEditorModeStyle.h"
 #include "Widgets/Input/SCheckBox.h"
 #include "Widgets/Layout/SBox.h"
 #include "ModelingWidgets/ModelingCustomizationUtil.h"
@@ -31,6 +32,12 @@ void FModelingToolsBrushSizeCustomization::CustomizeHeader(TSharedRef<IPropertyH
 
 	TSharedPtr<IPropertyHandle> WorldRadius = StructPropertyHandle->GetChildHandle(GET_MEMBER_NAME_CHECKED(FBrushToolRadius, WorldRadius));
 	WorldRadius->MarkHiddenByCustomization();
+	
+	TSharedPtr<IPropertyHandle> EnablePressureSensitivity = StructPropertyHandle->GetChildHandle(GET_MEMBER_NAME_CHECKED(FBrushToolRadius, bEnablePressureSensitivity));
+	EnablePressureSensitivity->MarkHiddenByCustomization();
+	
+	TSharedPtr<IPropertyHandle> SupportPressureSensitivity = StructPropertyHandle->GetChildHandle(GET_MEMBER_NAME_CHECKED(FBrushToolRadius, bToolSupportsPressureSensitivity));
+	SupportPressureSensitivity->MarkHiddenByCustomization();
 
 	auto GetCurrentSizeType = [SizeType]() -> EBrushToolSizeType
 	{
@@ -65,31 +72,34 @@ void FModelingToolsBrushSizeCustomization::CustomizeHeader(TSharedRef<IPropertyH
 		return (GetCurrentSizeType() == EBrushToolSizeType::Adaptive) ? TInterval<float>(0.0f, 1.0f) : TInterval<float>(1.0f, 1000.0f);
 	};
 
+	TSharedPtr<SHorizontalBox> Container;
+
 	HeaderRow
-		.OverrideResetToDefault(FResetToDefaultOverride::Hide())
-		.WholeRowContent()
+	.OverrideResetToDefault(FResetToDefaultOverride::Hide())
+	.WholeRowContent()
 	[
-		SNew(SHorizontalBox)
+		SAssignNew(Container, SHorizontalBox)
+	];
 
-		// name
-		+ SHorizontalBox::Slot()
+	// name
+	Container->AddSlot()
 		.AutoWidth()
 		.Padding(FMargin(0,4))
 		[
 			WrapInFixedWidthBox(AdaptiveSize->CreatePropertyNameWidget(), FSculptToolsUIConstants::SculptShortLabelWidth)
-		]
+		];
 
-		// slider
-		+ SHorizontalBox::Slot()
+	// slider
+	Container->AddSlot()
 		.Padding(ModelingUIConstants::LabelWidgetMinPadding, 0,0,0)
 		.FillWidth(10.0f)
 		[
 			SNew(SDynamicNumericEntry)
 			.Source(NumericSource)
-		]
+		];
 
-		// toggle
-		+ SHorizontalBox::Slot()
+	// toggle
+	Container->AddSlot()
 		.AutoWidth()
 		//.FillWidth(0.5f)
 		.Padding(FMargin(ModelingUIConstants::MultiWidgetRowHorzPadding, ModelingUIConstants::DetailRowVertPadding, 0, ModelingUIConstants::DetailRowVertPadding))
@@ -125,8 +135,50 @@ void FModelingToolsBrushSizeCustomization::CustomizeHeader(TSharedRef<IPropertyH
 						]
 					]
 			]
-		]
-	];
+		];
+
+	// pressure sensitivity toggle
+	bool bAllowPressureSensitivity;
+	SupportPressureSensitivity->GetValue(bAllowPressureSensitivity);
+	if (bAllowPressureSensitivity)
+	{
+		Container->AddSlot()
+			.AutoWidth()
+			//.FillWidth(0.5f)
+			.Padding(FMargin(ModelingUIConstants::MultiWidgetRowHorzPadding, ModelingUIConstants::DetailRowVertPadding, 0, ModelingUIConstants::DetailRowVertPadding))
+			.HAlign(HAlign_Right)
+			[
+				SNew(SBox)
+				.WidthOverride(30)
+				[
+					SNew(SCheckBox)
+					.Style(FAppStyle::Get(), "DetailsView.SectionButton")
+					.Padding(FMargin(4, 2))
+					.HAlign(HAlign_Center)
+					.ToolTipText( LOCTEXT("PressureSensitivityToggle_Label", "Toggle Pressure Sensitivity for Brush Size"))
+					.OnCheckStateChanged_Lambda([this, EnablePressureSensitivity](ECheckBoxState State) {
+						
+						EnablePressureSensitivity->SetValue( State == ECheckBoxState::Checked);
+					})
+					.IsChecked_Lambda([this, EnablePressureSensitivity]() {
+						bool CurValue;
+						EnablePressureSensitivity->GetValue(CurValue);
+						return CurValue ? ECheckBoxState::Checked : ECheckBoxState::Unchecked;
+					})
+					[
+						SNew(SHorizontalBox)
+						+ SHorizontalBox::Slot()
+						.VAlign(VAlign_Center)
+						.Padding(FMargin(0))
+						.AutoWidth()
+						[
+							SNew(SImage)
+							.Image(FModelingToolsEditorModeStyle::Get()->GetBrush("BrushIcons.PressureSensitivity"))
+						]
+					]
+			]
+		];
+	}
 
 }
 
diff --git a/Engine/Plugins/Editor/ModelingToolsEditorMode/Source/ModelingToolsEditorMode/Private/ModelingToolsEditorModeStyle.cpp b/Engine/Plugins/Editor/ModelingToolsEditorMode/Source/ModelingToolsEditorMode/Private/ModelingToolsEditorModeStyle.cpp
--- a/Engine/Plugins/Editor/ModelingToolsEditorMode/Source/ModelingToolsEditorMode/Private/ModelingToolsEditorModeStyle.cpp
+++ b/Engine/Plugins/Editor/ModelingToolsEditorMode/Source/ModelingToolsEditorMode/Private/ModelingToolsEditorModeStyle.cpp
@@ -479,6 +479,9 @@ void FModelingToolsEditorModeStyle::Initialize()
 		StyleSet->Set("BrushTypeIcons.HeightFlatten", new FSlateImageBrush(StyleSet->RootToCoreContentDir(TEXT("../Editor/Slate/Icons/Landscape/Landscape_Flatten_x40.png")), Icon20x20));
 		StyleSet->Set("BrushTypeIcons.HeightErase", new FSlateImageBrush(StyleSet->RootToCoreContentDir(TEXT("../Editor/Slate/Icons/Landscape/Landscape_Erase_x40.png")), Icon20x20));
 
+		// Pressure Sensitivity toggle button
+		StyleSet->Set("BrushIcons.PressureSensitivity", new IMAGE_BRUSH_SVG("Icons/PressureSensitivity", Icon16x16));
+
 		//
 		// Icons for selection buttons in PolyEd and TriEd
 		//
diff --git a/Engine/Plugins/Runtime/MeshModelingToolset/Source/MeshModelingTools/Private/MeshVertexPaintTool.cpp b/Engine/Plugins/Runtime/MeshModelingToolset/Source/MeshModelingTools/Private/MeshVertexPaintTool.cpp
--- a/Engine/Plugins/Runtime/MeshModelingToolset/Source/MeshModelingTools/Private/MeshVertexPaintTool.cpp
+++ b/Engine/Plugins/Runtime/MeshModelingToolset/Source/MeshModelingTools/Private/MeshVertexPaintTool.cpp
@@ -288,6 +288,7 @@ void UMeshVertexPaintTool::Setup()
 	UMeshSculptToolBase::BrushProperties->bShowFlowRate = true;
 	UMeshSculptToolBase::BrushProperties->bShowSpacing = false;
 	UMeshSculptToolBase::BrushProperties->BrushFalloffAmount = 0.66f;
+	UMeshSculptToolBase::BrushProperties->BrushSize.bToolSupportsPressureSensitivity = true;
 	CalculateBrushRadius();
 
 	PaintBrushOpProperties = NewObject<UVertexColorPaintBrushOpProps>(this);
@@ -828,6 +829,7 @@ bool UMeshVertexPaintTool::UpdateStampPosition(const FRay& WorldRay)
 	CurrentStamp = LastStamp;
 	CurrentStamp.DeltaTime = FMathd::Min((FDateTime::Now() - LastStamp.TimeStamp).GetTotalSeconds(), 1.0);
 	CurrentStamp.WorldFrame = GetBrushFrameWorld();
+	CurrentStamp.Radius = GetActiveBrushRadius();
 	CurrentStamp.LocalFrame = GetBrushFrameLocal();
 	CurrentStamp.Power = GetActivePressure() * GetCurrentBrushStrength();
 
diff --git a/Engine/Plugins/Runtime/MeshModelingToolset/Source/MeshModelingTools/Private/MeshVertexSculptTool.cpp b/Engine/Plugins/Runtime/MeshModelingToolset/Source/MeshModelingTools/Private/MeshVertexSculptTool.cpp
--- a/Engine/Plugins/Runtime/MeshModelingToolset/Source/MeshModelingTools/Private/MeshVertexSculptTool.cpp
+++ b/Engine/Plugins/Runtime/MeshModelingToolset/Source/MeshModelingTools/Private/MeshVertexSculptTool.cpp
@@ -216,6 +216,7 @@ void UMeshVertexSculptTool::Setup()
 	AddToolPropertySource(UMeshSculptToolBase::BrushProperties);
 	UMeshSculptToolBase::BrushProperties->bShowPerBrushProps = false;
 	UMeshSculptToolBase::BrushProperties->bShowFalloff = false;
+	UMeshSculptToolBase::BrushProperties->BrushSize.bToolSupportsPressureSensitivity = true;
 	SculptProperties->RestoreProperties(this, GetPropertyCacheIdentifier());
 	
 	AddToolPropertySource(SculptProperties);
@@ -914,6 +915,7 @@ bool UMeshVertexSculptTool::UpdateStampPosition(const FRay& WorldRay)
 	//CurrentStamp.DeltaTime = FMathd::Min((FDateTime::Now() - LastStamp.TimeStamp).GetTotalSeconds(), 1.0);
 	CurrentStamp.DeltaTime = 0.03;		// 30 fps - using actual time is no good now that we support variable stamps!
 	CurrentStamp.WorldFrame = GetBrushFrameWorld();
+	CurrentStamp.Radius = GetActiveBrushRadius();
 	CurrentStamp.LocalFrame = GetBrushFrameLocal();
 	CurrentStamp.Power = GetActivePressure() * GetCurrentBrushStrength();
 
diff --git a/Engine/Plugins/Runtime/MeshModelingToolset/Source/MeshModelingTools/Private/Sculpting/MeshSculptToolBase.cpp b/Engine/Plugins/Runtime/MeshModelingToolset/Source/MeshModelingTools/Private/Sculpting/MeshSculptToolBase.cpp
--- a/Engine/Plugins/Runtime/MeshModelingToolset/Source/MeshModelingTools/Private/Sculpting/MeshSculptToolBase.cpp
+++ b/Engine/Plugins/Runtime/MeshModelingToolset/Source/MeshModelingTools/Private/Sculpting/MeshSculptToolBase.cpp
@@ -659,6 +659,14 @@ FRay3d UMeshSculptToolBase::GetLocalRay(const FRay& WorldRay) const
 	return LocalRay;
 }
 
+bool UMeshSculptToolBase::GetBrushSizePressureSensitivityEnabled() const
+{
+	if (ToolPropertyObjects.Contains(BrushProperties) && BrushProperties->BrushSize.bToolSupportsPressureSensitivity)
+	{
+		return BrushProperties->BrushSize.bEnablePressureSensitivity;
+	}
+	return false;
+}
 
 
 void UMeshSculptToolBase::UpdateBrushFrameWorld(const FVector3d& NewPosition, const FVector3d& NewNormal)
@@ -978,6 +986,23 @@ void UMeshSculptToolBase::CalculateBrushRadius()
 	CurrentBrushRadius = BrushProperties->BrushSize.GetWorldRadius();
 }
 
+double UMeshSculptToolBase::GetActiveBrushRadius()
+{
+	if (GetBrushSizePressureSensitivityEnabled())
+	{
+		if (BrushProperties->BrushSize.SizeType == EBrushToolSizeType::Adaptive)
+		{
+			return 0.5 * BrushProperties->BrushSize.WorldSizeRange.Interpolate(BrushProperties->BrushSize.AdaptiveSize * GetActivePressure());
+		}
+		if (BrushProperties->BrushSize.SizeType == EBrushToolSizeType::World)
+		{
+			const TInterval<float> WorldSizeRange = BrushProperties->BrushSize.WorldSizeRange;
+			return FMath::Lerp(WorldSizeRange.Min, BrushProperties->BrushSize.GetWorldRadius(), GetActivePressure());
+		}
+	}
+	return CurrentBrushRadius;
+}
+
 
 
 double UMeshSculptToolBase::GetCurrentBrushStrength()
diff --git a/Engine/Plugins/Runtime/MeshModelingToolset/Source/MeshModelingTools/Public/Sculpting/MeshSculptToolBase.h b/Engine/Plugins/Runtime/MeshModelingToolset/Source/MeshModelingTools/Public/Sculpting/MeshSculptToolBase.h
--- a/Engine/Plugins/Runtime/MeshModelingToolset/Source/MeshModelingTools/Public/Sculpting/MeshSculptToolBase.h
+++ b/Engine/Plugins/Runtime/MeshModelingToolset/Source/MeshModelingTools/Public/Sculpting/MeshSculptToolBase.h
@@ -55,6 +55,12 @@ struct MESHMODELINGTOOLS_API FBrushToolRadius
 		DisplayName = "World Radius", UIMin = "1.0", UIMax = "1000.0", ClampMin = "0.1", ClampMax = "50000.0"))
 	float WorldRadius = 100.0;
 
+	UPROPERTY(EditAnywhere, Category = Brush)
+	bool bToolSupportsPressureSensitivity = false;
+	
+	UPROPERTY(EditAnywhere, Category = Brush, meta = (EditCondition = "bToolSupportsPressureSensitivity"))
+	bool bEnablePressureSensitivity = false;
+
 	/**
 	 * WorldSizeRange defines the min/max dimensions for Adaptive brush size
 	 */
@@ -285,6 +291,8 @@ protected:
 
 	FRay3d GetLocalRay(const FRay& WorldRay) const;
 
+	/** Returns true if BrushSize Pressure Sensitivity is both supported by the tool and currently enabled. */
+	bool GetBrushSizePressureSensitivityEnabled() const;
 
 	/**
 	 * Subclass must implement this and return relevant rendering component
@@ -390,6 +398,7 @@ protected:
 	double CurrentBrushRadius = 1.0;
 	virtual void InitializeBrushSizeRange(const UE::Geometry::FAxisAlignedBox3d& TargetBounds);
 	virtual void CalculateBrushRadius();
+	virtual double GetActiveBrushRadius();
 	virtual double GetCurrentBrushRadius() const { return CurrentBrushRadius; }
 
 	double CurrentBrushFalloff = 0.5;
