Modeling Mode: add 'Expand to All Connected' functionality to viewport toolbar

diff --git a/Engine/Plugins/Editor/ModelingToolsEditorMode/Source/ModelingToolsEditorMode/Private/ModelingToolsEditorMode.cpp b/Engine/Plugins/Editor/ModelingToolsEditorMode/Source/ModelingToolsEditorMode/Private/ModelingToolsEditorMode.cpp
--- a/Engine/Plugins/Editor/ModelingToolsEditorMode/Source/ModelingToolsEditorMode/Private/ModelingToolsEditorMode.cpp
+++ b/Engine/Plugins/Editor/ModelingToolsEditorMode/Source/ModelingToolsEditorMode/Private/ModelingToolsEditorMode.cpp
@@ -1440,6 +1440,9 @@ void UModelingToolsEditorMode::OnBlueprintPreCompile(UBlueprint* Blueprint)
 					FToolMenuEntry::InitMenuEntryWithCommandList(
 						FModelingToolsManagerCommands::Get().BeginSelectionAction_SelectAll, CommandList));
 					SelectionEditsSection.AddEntry(
+					FToolMenuEntry::InitMenuEntryWithCommandList(
+						FModelingToolsManagerCommands::Get().BeginSelectionAction_ExpandToConnected, CommandList));
+					SelectionEditsSection.AddEntry(
 					FToolMenuEntry::InitMenuEntryWithCommandList(
 						FModelingToolsManagerCommands::Get().BeginSelectionAction_Invert, CommandList));
 					SelectionEditsSection.AddEntry(
