Modeling Tools: Add paint/erase pressure sensitivity to vertex paint tool

diff --git a/Engine/Plugins/Editor/ModelingToolsEditorMode/Source/ModelingToolsEditorMode/Private/DetailsCustomizations/MeshVertexPaintToolCustomizations.cpp b/Engine/Plugins/Editor/ModelingToolsEditorMode/Source/ModelingToolsEditorMode/Private/DetailsCustomizations/MeshVertexPaintToolCustomizations.cpp
new file mode 100644
--- /dev/null
+++ b/Engine/Plugins/Editor/ModelingToolsEditorMode/Source/ModelingToolsEditorMode/Private/DetailsCustomizations/MeshVertexPaintToolCustomizations.cpp
@@ -0,0 +1,93 @@
+// Copyright Epic Games, Inc. All Rights Reserved.
+
+#include "DetailsCustomizations/MeshVertexPaintToolCustomizations.h"
+
+#include "MeshVertexPaintTool.h"
+#include "DetailLayoutBuilder.h"
+#include "DetailWidgetRow.h"
+#include "ModelingToolsEditorModeStyle.h"
+#include "Widgets/Images/SImage.h"
+#include "Widgets/Input/SCheckBox.h"
+
+TSharedRef<IDetailCustomization> FVertexPaintBasicPropertiesDetails::MakeInstance()
+{
+	return MakeShareable(new FVertexPaintBasicPropertiesDetails);
+}
+
+void FVertexPaintBasicPropertiesDetails::CustomizeDetails(IDetailLayoutBuilder& DetailBuilder)
+{
+	TSharedPtr<IPropertyHandle> PaintHandle = DetailBuilder.GetProperty(GET_MEMBER_NAME_CHECKED(UVertexPaintBasicProperties, PaintColor), UVertexPaintBasicProperties::StaticClass());
+	ensure(PaintHandle->IsValidHandle());
+	
+	TSharedPtr<IPropertyHandle> PaintPressureHandle = DetailBuilder.GetProperty(GET_MEMBER_NAME_CHECKED(UVertexPaintBasicProperties, bIsPaintPressureEnabled), UVertexPaintBasicProperties::StaticClass());
+	ensure(PaintPressureHandle->IsValidHandle());
+	PaintPressureHandle->MarkHiddenByCustomization();
+
+	TSharedPtr<IPropertyHandle> EraseHandle = DetailBuilder.GetProperty(GET_MEMBER_NAME_CHECKED(UVertexPaintBasicProperties, EraseColor), UVertexPaintBasicProperties::StaticClass());
+	ensure(EraseHandle->IsValidHandle());
+
+	TSharedPtr<IPropertyHandle> ErasePressureHandle = DetailBuilder.GetProperty(GET_MEMBER_NAME_CHECKED(UVertexPaintBasicProperties, bIsErasePressureEnabled), UVertexPaintBasicProperties::StaticClass());
+	ensure(ErasePressureHandle->IsValidHandle());
+	ErasePressureHandle->MarkHiddenByCustomization();
+
+	BuildPaintPressureWidget(DetailBuilder, PaintHandle, PaintPressureHandle);
+	BuildPaintPressureWidget(DetailBuilder, EraseHandle, ErasePressureHandle);
+}
+
+void FVertexPaintBasicPropertiesDetails::BuildPaintPressureWidget(IDetailLayoutBuilder& DetailBuilder,
+	const TSharedPtr<IPropertyHandle>& PropHandle, TSharedPtr<IPropertyHandle> EnablePressureSensitivityHandle)
+{
+	IDetailPropertyRow* DetailRow = DetailBuilder.EditDefaultProperty(PropHandle);
+	TSharedPtr<SWidget> NameWidget, ValueWidget;
+	DetailRow->GetDefaultWidgets(NameWidget, ValueWidget);
+
+	TSharedPtr<SHorizontalBox> ValueContent;
+
+	DetailRow->CustomWidget()
+	.NameContent() // property text
+	[
+		NameWidget->AsShared()
+	]
+	.ValueContent() // normal property widget and pressure sensitivity toggle
+	[
+		SAssignNew(ValueContent, SHorizontalBox)
+	];
+
+	// always add usual widget
+	ValueContent->AddSlot()
+	[
+		ValueWidget->AsShared()
+	];
+
+	// add pressure sensitivity toggle
+	ValueContent->AddSlot()
+	.AutoWidth()
+	[
+		SNew(SCheckBox)
+		.Style(FAppStyle::Get(), "DetailsView.SectionButton")
+		.Padding(FMargin(4, 2))
+		.ToolTipText(EnablePressureSensitivityHandle->GetToolTipText())
+		.HAlign(HAlign_Center)
+		.OnCheckStateChanged_Lambda([EnablePressureSensitivityHandle](const ECheckBoxState NewState)
+		{
+			EnablePressureSensitivityHandle->SetValue(NewState == ECheckBoxState::Checked);
+		})
+		.IsChecked_Lambda([EnablePressureSensitivityHandle]() -> ECheckBoxState
+		{
+			bool bSet;
+			EnablePressureSensitivityHandle->GetValue(bSet);
+			return bSet ? ECheckBoxState::Checked : ECheckBoxState::Unchecked;
+		})
+		[
+			SNew(SHorizontalBox)
+			+ SHorizontalBox::Slot()
+			.VAlign(VAlign_Center)
+			.Padding(FMargin(0))
+			.AutoWidth()
+			[
+				SNew(SImage)
+				.Image(FModelingToolsEditorModeStyle::Get()->GetBrush("BrushIcons.PressureSensitivity"))
+			]
+		]
+	];
+}
\ No newline at end of file
diff --git a/Engine/Plugins/Editor/ModelingToolsEditorMode/Source/ModelingToolsEditorMode/Private/ModelingToolsEditorMode.cpp b/Engine/Plugins/Editor/ModelingToolsEditorMode/Source/ModelingToolsEditorMode/Private/ModelingToolsEditorMode.cpp
--- a/Engine/Plugins/Editor/ModelingToolsEditorMode/Source/ModelingToolsEditorMode/Private/ModelingToolsEditorMode.cpp
+++ b/Engine/Plugins/Editor/ModelingToolsEditorMode/Source/ModelingToolsEditorMode/Private/ModelingToolsEditorMode.cpp
@@ -553,7 +553,9 @@ void UModelingToolsEditorMode::Enter()
 	auto MeshGroupPaintToolBuilder = NewObject<UMeshGroupPaintToolBuilder>();
 	MeshGroupPaintToolBuilder->StylusAPI = StylusAPI;
 	RegisterTool(ToolManagerCommands.BeginMeshGroupPaintTool, TEXT("BeginMeshGroupPaintTool"), MeshGroupPaintToolBuilder);
-	RegisterTool(ToolManagerCommands.BeginMeshVertexPaintTool, TEXT("BeginMeshVertexPaintTool"), NewObject<UMeshVertexPaintToolBuilder>());
+	UMeshVertexPaintToolBuilder* MeshVertexPaintToolBuilder = NewObject<UMeshVertexPaintToolBuilder>();
+	MeshVertexPaintToolBuilder->StylusAPI = StylusAPI;
+	RegisterTool(ToolManagerCommands.BeginMeshVertexPaintTool, TEXT("BeginMeshVertexPaintTool"), MeshVertexPaintToolBuilder);
 
 	RegisterTool(ToolManagerCommands.BeginPolyEditTool, TEXT("BeginPolyEditTool"), NewObject<UEditMeshPolygonsToolBuilder>());
 	UEditMeshPolygonsToolBuilder* TriEditBuilder = NewObject<UEditMeshPolygonsToolBuilder>();
diff --git a/Engine/Plugins/Editor/ModelingToolsEditorMode/Source/ModelingToolsEditorMode/Private/ModelingToolsEditorModeModule.cpp b/Engine/Plugins/Editor/ModelingToolsEditorMode/Source/ModelingToolsEditorMode/Private/ModelingToolsEditorModeModule.cpp
--- a/Engine/Plugins/Editor/ModelingToolsEditorMode/Source/ModelingToolsEditorMode/Private/ModelingToolsEditorModeModule.cpp
+++ b/Engine/Plugins/Editor/ModelingToolsEditorMode/Source/ModelingToolsEditorMode/Private/ModelingToolsEditorModeModule.cpp
@@ -12,6 +12,7 @@
 #include "DetailsCustomizations/ModelingToolPropertyCustomizations.h"
 #include "DetailsCustomizations/ModelingToolsBrushSizeCustomization.h"
 #include "DetailsCustomizations/MeshVertexSculptToolCustomizations.h"
+#include "DetailsCustomizations/MeshVertexPaintToolCustomizations.h"
 #include "DetailsCustomizations/BakeMeshAttributeToolCustomizations.h"
 #include "DetailsCustomizations/BakeTransformToolCustomizations.h"
 #include "DetailsCustomizations/MeshTopologySelectionMechanicCustomization.h"
@@ -20,6 +21,7 @@
 #include "PropertySets/ColorChannelFilterPropertyType.h"
 #include "Selection/MeshTopologySelectionMechanic.h"
 #include "MeshVertexSculptTool.h"
+#include "MeshVertexPaintTool.h"
 #include "BakeMeshAttributeMapsTool.h"
 #include "BakeMultiMeshAttributeMapsTool.h"
 #include "BakeMeshAttributeVertexTool.h"
@@ -134,6 +136,10 @@ void FModelingToolsEditorModeModule::OnPostEngineInit()
 	PropertyModule.RegisterCustomClassLayout("TwistKelvinletBrushOpProps", FOnGetDetailCustomizationInstance::CreateStatic(&FMeshSculptBrushOpPropertiesDetails<UTwistKelvinletBrushOpProps>::MakeInstance));
 	ClassesToUnregisterOnShutdown.Add(UTwistKelvinletBrushOpProps::StaticClass()->GetFName());
 
+	// Paint
+	PropertyModule.RegisterCustomClassLayout("VertexPaintBasicProperties", FOnGetDetailCustomizationInstance::CreateStatic(&FVertexPaintBasicPropertiesDetails::MakeInstance));
+	ClassesToUnregisterOnShutdown.Add(UVertexPaintBasicProperties::StaticClass()->GetFName());
+
 	/// Bake
 	PropertyModule.RegisterCustomClassLayout("BakeMeshAttributeMapsToolProperties", FOnGetDetailCustomizationInstance::CreateStatic(&FBakeMeshAttributeMapsToolDetails::MakeInstance));
 	ClassesToUnregisterOnShutdown.Add(UBakeMeshAttributeMapsToolProperties::StaticClass()->GetFName());
diff --git a/Engine/Plugins/Editor/ModelingToolsEditorMode/Source/ModelingToolsEditorMode/Public/DetailsCustomizations/MeshVertexPaintToolCustomizations.h b/Engine/Plugins/Editor/ModelingToolsEditorMode/Source/ModelingToolsEditorMode/Public/DetailsCustomizations/MeshVertexPaintToolCustomizations.h
new file mode 100644
--- /dev/null
+++ b/Engine/Plugins/Editor/ModelingToolsEditorMode/Source/ModelingToolsEditorMode/Public/DetailsCustomizations/MeshVertexPaintToolCustomizations.h
@@ -0,0 +1,24 @@
+// Copyright Epic Games, Inc. All Rights Reserved.
+
+#pragma once
+
+#include "IDetailCustomization.h"
+
+class IDetailLayoutBuilder;
+class IPropertyHandle;
+
+// customization for UVertexPaintBasicProperties
+// creates custom widgets for paint and erase brushes with pressure sensitivity
+class FVertexPaintBasicPropertiesDetails : public IDetailCustomization
+{
+public:
+	static TSharedRef<IDetailCustomization> MakeInstance();
+	void CustomizeDetails(IDetailLayoutBuilder& DetailBuilder) override;
+private:
+	// helper to build strength widget, with pressure sensitivity toggle
+	static void BuildPaintPressureWidget(
+		IDetailLayoutBuilder& DetailBuilder,
+		const TSharedPtr<IPropertyHandle>& ColorPropHandle,
+		TSharedPtr<IPropertyHandle> EnablePressureSensitivityHandle
+	);
+};
\ No newline at end of file
diff --git a/Engine/Plugins/Runtime/MeshModelingToolset/Source/MeshModelingTools/Private/MeshVertexPaintTool.cpp b/Engine/Plugins/Runtime/MeshModelingToolset/Source/MeshModelingTools/Private/MeshVertexPaintTool.cpp
--- a/Engine/Plugins/Runtime/MeshModelingToolset/Source/MeshModelingTools/Private/MeshVertexPaintTool.cpp
+++ b/Engine/Plugins/Runtime/MeshModelingToolset/Source/MeshModelingTools/Private/MeshVertexPaintTool.cpp
@@ -264,6 +264,10 @@ void UMeshVertexPaintTool::Setup()
 		[this](const FModelingToolsColorChannelFilter&) { OnChannelFilterModified(); });
 	BasicProperties->WatchProperty(BasicProperties->SecondaryActionType,
 		[this](EMeshVertexPaintSecondaryActionType NewType) { UpdateSecondaryBrushType(NewType); });
+	BasicProperties->WatchProperty(BasicProperties->bIsErasePressureEnabled,
+		[this](const bool bEnabled) { EraseBrushOpProperties->bIsStrengthPressureEnabled = bEnabled; });
+	BasicProperties->WatchProperty(BasicProperties->bIsPaintPressureEnabled,
+		[this](const bool bEnabled) { PaintBrushOpProperties->bIsStrengthPressureEnabled = bEnabled; });
 
 	InitializeIndicator();
 
@@ -570,7 +574,7 @@ void UMeshVertexPaintTool::OnBeginStroke(const FRay& WorldRay)
 	LastStamp.Falloff = GetCurrentBrushFalloff();
 	LastStamp.Direction = GetInInvertStroke() ? -1.0 : 1.0;
 	LastStamp.Depth = GetCurrentBrushDepth();
-	LastStamp.Power = GetActivePressure() * GetCurrentBrushStrength();
+	LastStamp.Power = GetActiveBrushStrength();
 	LastStamp.TimeStamp = FDateTime::Now();
 
 	FSculptBrushOptions SculptOptions;
@@ -810,7 +814,7 @@ bool UMeshVertexPaintTool::UpdateStampPosition(const FRay& WorldRay)
 	CurrentStamp.WorldFrame = GetBrushFrameWorld();
 	CurrentStamp.Radius = GetActiveBrushRadius();
 	CurrentStamp.LocalFrame = GetBrushFrameLocal();
-	CurrentStamp.Power = GetActivePressure() * GetCurrentBrushStrength();
+	CurrentStamp.Power = GetActiveBrushStrength();
 
 	CurrentStamp.PrevLocalFrame = LastStamp.LocalFrame;
 	CurrentStamp.PrevWorldFrame = LastStamp.WorldFrame;
diff --git a/Engine/Plugins/Runtime/MeshModelingToolset/Source/MeshModelingTools/Public/MeshVertexPaintTool.h b/Engine/Plugins/Runtime/MeshModelingToolset/Source/MeshModelingTools/Public/MeshVertexPaintTool.h
--- a/Engine/Plugins/Runtime/MeshModelingToolset/Source/MeshModelingTools/Public/MeshVertexPaintTool.h
+++ b/Engine/Plugins/Runtime/MeshModelingToolset/Source/MeshModelingTools/Public/MeshVertexPaintTool.h
@@ -184,6 +184,10 @@ public:
 	UPROPERTY(EditAnywhere, Category = Settings, meta = ())
 	FLinearColor PaintColor = FLinearColor::Red;
 
+	/** Should pressure sensitivity affect Paint brush strength? */
+	UPROPERTY(EditAnywhere, Category = Settings)
+	bool bIsPaintPressureEnabled = true;
+
 	/** Controls how painted Colors will be combined with the existing Colors */
 	UPROPERTY(EditAnywhere, Category = Settings)
 	EMeshVertexPaintColorBlendMode BlendMode = EMeshVertexPaintColorBlendMode::Lerp;
@@ -196,6 +200,10 @@ public:
 	UPROPERTY(EditAnywhere, Category = Settings, meta = (EditConditionHides, EditCondition = "SecondaryActionType != EMeshVertexPaintSecondaryActionType::Smooth"))
 	FLinearColor EraseColor = FLinearColor::White;
 
+	/** Should pressure sensitivity affect Erase brush strength? */
+	UPROPERTY(EditAnywhere, Category = Settings, meta = (EditConditionHides, EditCondition = "SecondaryActionType != EMeshVertexPaintSecondaryActionType::Smooth"))
+	bool bIsErasePressureEnabled = true;
+
 	/** Strength of Smooth Brush */
 	UPROPERTY(EditAnywhere, Category = Settings, meta = (UIMin = 0, UIMax = 1, EditConditionHides, EditCondition = "SecondaryActionType == EMeshVertexPaintSecondaryActionType::Smooth"))
 	float SmoothStrength = 0.25;
diff --git a/Engine/Plugins/Runtime/MeshModelingToolset/Source/MeshModelingTools/Public/Sculpting/MeshVertexPaintBrushOps.h b/Engine/Plugins/Runtime/MeshModelingToolset/Source/MeshModelingTools/Public/Sculpting/MeshVertexPaintBrushOps.h
--- a/Engine/Plugins/Runtime/MeshModelingToolset/Source/MeshModelingTools/Public/Sculpting/MeshVertexPaintBrushOps.h
+++ b/Engine/Plugins/Runtime/MeshModelingToolset/Source/MeshModelingTools/Public/Sculpting/MeshVertexPaintBrushOps.h
@@ -137,7 +137,7 @@ public:
 	UPROPERTY(EditAnywhere, Category = PaintBrush)
 	bool bApplyFalloff = true;
 
-
+	virtual bool SupportsStrengthPressure() override { return true; };
 	virtual float GetStrength() override { return Strength; }
 	virtual void SetStrength(float NewStrength) override { Strength = FMathf::Clamp(NewStrength, 0.0f, 1.0f); }
 	virtual float GetFalloff() override { return Falloff; }
