Modeling Mode: Allow the Vertex Sculpt tool to define and work with currently 'available' brushes out of all registered brushes

diff --git a/Engine/Plugins/Editor/ModelingToolsEditorMode/Source/ModelingToolsEditorMode/Private/DetailsCustomizations/MeshVertexSculptToolCustomizations.cpp b/Engine/Plugins/Editor/ModelingToolsEditorMode/Source/ModelingToolsEditorMode/Private/DetailsCustomizations/MeshVertexSculptToolCustomizations.cpp
--- a/Engine/Plugins/Editor/ModelingToolsEditorMode/Source/ModelingToolsEditorMode/Private/DetailsCustomizations/MeshVertexSculptToolCustomizations.cpp
+++ b/Engine/Plugins/Editor/ModelingToolsEditorMode/Source/ModelingToolsEditorMode/Private/DetailsCustomizations/MeshVertexSculptToolCustomizations.cpp
@@ -254,7 +254,7 @@ void FVertexBrushSculptPropertiesDetails::CustomizeDetails(IDetailLayoutBuilder&
 	UMeshVertexSculptTool* Tool = BrushProperties->Tool.Get();
 	TargetTool = Tool;
 
-	const TSet<UMeshSculptToolBase::FBrushTypeInfo> BrushTypeInfos = Tool->GetRegisteredPrimaryBrushTypes();
+	const TSet<UMeshSculptToolBase::FBrushTypeInfo> BrushTypeInfos = Tool->GetAvailableBrushTypes();
 	int32 CurrentBrushType = (int32)BrushProperties->PrimaryBrushID;
 	int32 CurrentBrushTypeIndex = 0;
 
diff --git a/Engine/Plugins/Runtime/MeshModelingToolset/Source/MeshModelingTools/Public/MeshVertexSculptTool.h b/Engine/Plugins/Runtime/MeshModelingToolset/Source/MeshModelingTools/Public/MeshVertexSculptTool.h
--- a/Engine/Plugins/Runtime/MeshModelingToolset/Source/MeshModelingTools/Public/MeshVertexSculptTool.h
+++ b/Engine/Plugins/Runtime/MeshModelingToolset/Source/MeshModelingTools/Public/MeshVertexSculptTool.h
@@ -334,6 +334,9 @@ public:
 	virtual void SetActiveFalloffType(int32 Identifier);
 	virtual void SetRegionFilterType(int32 Identifier);
 
+	// retrieves a set of FBrushTypeInfos representing the brushes currently available to be enabled
+	virtual TSet<UMeshSculptToolBase::FBrushTypeInfo> GetAvailableBrushTypes() { return GetRegisteredPrimaryBrushTypes(); };
+
 	void SetDefaultPrimaryBrushID(int32 InPrimaryBrushID);
 
 	bool CanUpdateBrushType() const;
