ModelingTools: Fixed issue where editor gizmos were hidden in modeling mode when selection filter was set to any component type (vert, edge, tri).

diff --git a/Engine/Plugins/Editor/ModelingToolsEditorMode/Source/ModelingToolsEditorMode/Private/ModelingToolsEditorMode.cpp b/Engine/Plugins/Editor/ModelingToolsEditorMode/Source/ModelingToolsEditorMode/Private/ModelingToolsEditorMode.cpp
--- a/Engine/Plugins/Editor/ModelingToolsEditorMode/Source/ModelingToolsEditorMode/Private/ModelingToolsEditorMode.cpp
+++ b/Engine/Plugins/Editor/ModelingToolsEditorMode/Source/ModelingToolsEditorMode/Private/ModelingToolsEditorMode.cpp
@@ -260,7 +260,7 @@ bool UModelingToolsEditorMode::ShouldDrawWidget() const
 	}
 
 	// hide standard xform gizmo if we have an active selection
-	if (GetSelectionManager() && (GetSelectionManager()->HasSelection() || GetSelectionManager()->GetMeshTopologyMode() != UGeometrySelectionManager::EMeshTopologyMode::None))
+	if (GetSelectionManager() && (GetSelectionManager()->HasSelection() && GetSelectionManager()->GetMeshTopologyMode() != UGeometrySelectionManager::EMeshTopologyMode::None))
 	{
 		return false;
 	}
