ModelingTools: Fixed crash when swapping from Modeling Mode to Animation Mode with Motion Trails active.

diff --git a/Engine/Plugins/Editor/ModelingToolsEditorMode/Source/ModelingToolsEditorMode/Private/ModelingToolsEditorModeToolkit.cpp b/Engine/Plugins/Editor/ModelingToolsEditorMode/Source/ModelingToolsEditorMode/Private/ModelingToolsEditorModeToolkit.cpp
--- a/Engine/Plugins/Editor/ModelingToolsEditorMode/Source/ModelingToolsEditorMode/Private/ModelingToolsEditorModeToolkit.cpp
+++ b/Engine/Plugins/Editor/ModelingToolsEditorMode/Source/ModelingToolsEditorMode/Private/ModelingToolsEditorModeToolkit.cpp
@@ -2325,6 +2325,13 @@ void FModelingToolsEditorModeToolkit::OnToolStarted(UInteractiveToolManager* Man
 	UpdateActiveToolProperties();
 
 	UInteractiveTool* CurTool = GetScriptableEditorMode()->GetToolManager(EToolsContextScope::EdMode)->GetActiveTool(EToolSide::Left);
+	// This can happen due to out of order mode Enter/Exit as a result of Exit being deferred to Tick.
+	// TODO: Swap this to an ensure once we fix the deferred UEdMode::Exit handling.
+	if (!CurTool)
+	{
+		return;
+	}
+	
 	CurTool->OnPropertySetsModified.AddSP(this, &FModelingToolsEditorModeToolkit::UpdateActiveToolProperties);
 	CurTool->OnPropertyModifiedDirectlyByTool.AddSP(this, &FModelingToolsEditorModeToolkit::InvalidateCachedDetailPanelState);
 
