ModelingSelectionInteraction: Fix ensure when initializing this selection interaction which occurs when there are more than 2 selection interactions at once.

This can occur when entering modeling mode while the modeless subsystem is active. The issue is that we give an explicit instance ID to the transform gizmo. These IDs need to be unique. We do not actually need to give it an ID because we never refer to it by this ID again, we destroy the gizmo explicitly.

diff --git a/Engine/Plugins/Editor/ModelingToolsEditorMode/Source/ModelingToolsEditorMode/Private/ModelingSelectionInteraction.cpp b/Engine/Plugins/Editor/ModelingToolsEditorMode/Source/ModelingToolsEditorMode/Private/ModelingSelectionInteraction.cpp
--- a/Engine/Plugins/Editor/ModelingToolsEditorMode/Source/ModelingToolsEditorMode/Private/ModelingSelectionInteraction.cpp
+++ b/Engine/Plugins/Editor/ModelingToolsEditorMode/Source/ModelingToolsEditorMode/Private/ModelingSelectionInteraction.cpp
@@ -59,7 +59,7 @@ void UModelingSelectionInteraction::Initialize(
 	TransformProxy = NewObject<UTransformProxy>(this);
 	TransformGizmo = UE::TransformGizmoUtil::CreateCustomRepositionableTransformGizmo( 
 		SelectionManager->GetToolsContext()->GizmoManager, ETransformGizmoSubElements::FullTranslateRotateScale,
-		this, TEXT("ModelingSelectionInteraction") );
+		this );
 	TransformGizmo->SetIsNonUniformScaleAllowedFunction([]() {return true;});
 	TransformGizmo->SetActiveTarget(TransformProxy, SelectionManager->GetToolsContext()->GizmoManager);
 	TransformGizmo->SetVisibility(false);
