ModelingTools: Fixed rare crash hovering over Mesh Element Selection before viewport is focused.

diff --git a/Engine/Plugins/Editor/ModelingToolsEditorMode/Source/ModelingToolsEditorMode/Private/ModelingToolsEditorMode.cpp b/Engine/Plugins/Editor/ModelingToolsEditorMode/Source/ModelingToolsEditorMode/Private/ModelingToolsEditorMode.cpp
--- a/Engine/Plugins/Editor/ModelingToolsEditorMode/Source/ModelingToolsEditorMode/Private/ModelingToolsEditorMode.cpp
+++ b/Engine/Plugins/Editor/ModelingToolsEditorMode/Source/ModelingToolsEditorMode/Private/ModelingToolsEditorMode.cpp
@@ -1600,11 +1600,13 @@ bool UModelingToolsEditorMode::TestForEditorGizmoHit(const FInputDeviceRay& Clic
 	// we need to filter out clicks on the gizmo. This function can do this check.
 	if ( ShouldDrawWidget() )
 	{
-		FViewport* FocusedViewport = GetToolManager()->GetContextQueriesAPI()->GetFocusedViewport();
-		HHitProxy* HitResult = FocusedViewport->GetHitProxy(ClickPos.ScreenPosition.X, ClickPos.ScreenPosition.Y);
-		if (HitResult && HitResult->IsA(HWidgetAxis::StaticGetType()))
+		if (FViewport* FocusedViewport = GetToolManager()->GetContextQueriesAPI()->GetFocusedViewport())
 		{
-			return true;
+			HHitProxy* HitResult = FocusedViewport->GetHitProxy(ClickPos.ScreenPosition.X, ClickPos.ScreenPosition.Y);
+			if (HitResult && HitResult->IsA(HWidgetAxis::StaticGetType()))
+			{
+				return true;
+			}
 		}
 	}
 	return false;
