ModelingTools: Fixed issues with switching between tools with default brush IDs.

diff --git a/Engine/Plugins/Editor/ModelingToolsEditorMode/Source/ModelingToolsEditorMode/Private/DetailsCustomizations/MeshVertexSculptToolCustomizations.cpp b/Engine/Plugins/Editor/ModelingToolsEditorMode/Source/ModelingToolsEditorMode/Private/DetailsCustomizations/MeshVertexSculptToolCustomizations.cpp
--- a/Engine/Plugins/Editor/ModelingToolsEditorMode/Source/ModelingToolsEditorMode/Private/DetailsCustomizations/MeshVertexSculptToolCustomizations.cpp
+++ b/Engine/Plugins/Editor/ModelingToolsEditorMode/Source/ModelingToolsEditorMode/Private/DetailsCustomizations/MeshVertexSculptToolCustomizations.cpp
@@ -294,7 +294,10 @@ void FVertexBrushSculptPropertiesDetails::CustomizeDetails(IDetailLayoutBuilder&
 		.OnSelectionChanged_Lambda([this](TSharedPtr<SComboPanel::FComboPanelItem> NewSelectedItem) {
 			if ( UMeshVertexSculptTool* Tool = TargetTool.Get() )
 			{
-				TargetTool.Get()->SetActiveBrushType(NewSelectedItem->Identifier);
+				if (Tool->CanUpdateBrushType())
+				{
+					Tool->SetActiveBrushType(NewSelectedItem->Identifier);
+				}
 			}
 		})
 		.FlyoutHeaderText(LOCTEXT("BrushesHeader", "Brush Types"))
