Fixed issue where tools in Modeling Mode and Scriptable Tools Mode would not tick if the viewport did not receive focus yet during the editor session.

diff --git a/Engine/Plugins/Editor/ModelingToolsEditorMode/Source/ModelingToolsEditorMode/Private/ModelingToolsEditorMode.cpp b/Engine/Plugins/Editor/ModelingToolsEditorMode/Source/ModelingToolsEditorMode/Private/ModelingToolsEditorMode.cpp
--- a/Engine/Plugins/Editor/ModelingToolsEditorMode/Source/ModelingToolsEditorMode/Private/ModelingToolsEditorMode.cpp
+++ b/Engine/Plugins/Editor/ModelingToolsEditorMode/Source/ModelingToolsEditorMode/Private/ModelingToolsEditorMode.cpp
@@ -355,6 +355,16 @@ void UModelingToolsEditorMode::Enter()
 			StaticCastSharedRef<FModelingToolsEditorModeToolkit>(Toolkit.ToSharedRef()));
 	}
 
+	// Ensure the ModeManager has initialized a focused viewport client otherwise the tool manager will not Tick.
+	FEditorModeTools& ModeManager = GLevelEditorModeTools();
+	if (ModeManager.GetFocusedViewportClient() == nullptr && GEditor)
+	{
+		if (FViewport* ActiveViewport = GEditor->GetActiveViewport())
+		{
+			ModeManager.ReceivedFocus(static_cast<FEditorViewportClient*>(ActiveViewport->GetClient()), ActiveViewport);
+		}
+	}
+
 	// set up SelectionManager and register known factory types
 	SelectionManager = NewObject<UGeometrySelectionManager>(GetToolManager());
 	SelectionManager->Initialize(GetInteractiveToolsContext(), GetToolManager()->GetContextTransactionsAPI());
diff --git a/Engine/Plugins/Editor/ScriptableToolsEditorMode/Source/ScriptableToolsEditorMode/Private/ScriptableToolsEditorMode.cpp b/Engine/Plugins/Editor/ScriptableToolsEditorMode/Source/ScriptableToolsEditorMode/Private/ScriptableToolsEditorMode.cpp
--- a/Engine/Plugins/Editor/ScriptableToolsEditorMode/Source/ScriptableToolsEditorMode/Private/ScriptableToolsEditorMode.cpp
+++ b/Engine/Plugins/Editor/ScriptableToolsEditorMode/Source/ScriptableToolsEditorMode/Private/ScriptableToolsEditorMode.cpp
@@ -181,6 +181,16 @@ void UScriptableToolsEditorMode::Enter()
 	// enable realtime viewport override
 	ConfigureRealTimeViewportsOverride(true);
 
+	// Ensure the ModeManager has initialized a focused viewport client otherwise the tool manager will not Tick.
+	FEditorModeTools& ModeManager = GLevelEditorModeTools();
+	if (ModeManager.GetFocusedViewportClient() == nullptr && GEditor)
+	{
+		if (FViewport* ActiveViewport = GEditor->GetActiveViewport())
+		{
+			ModeManager.ReceivedFocus(static_cast<FEditorViewportClient*>(ActiveViewport->GetClient()), ActiveViewport);
+		}
+	}
+
 	ScriptableTools = NewObject<UScriptableToolSet>(this);
 
 	UScriptableToolsModeCustomizationSettings* ModeSettings = GetMutableDefault<UScriptableToolsModeCustomizationSettings>();
