MeshModelingTools: Prevent tools from runing on internal/hidden dynamic mesh components.

diff --git a/Engine/Plugins/Editor/ModelingToolsEditorMode/Source/ModelingToolsEditorMode/Private/ModelingToolsEditorMode.cpp b/Engine/Plugins/Editor/ModelingToolsEditorMode/Source/ModelingToolsEditorMode/Private/ModelingToolsEditorMode.cpp
--- a/Engine/Plugins/Editor/ModelingToolsEditorMode/Source/ModelingToolsEditorMode/Private/ModelingToolsEditorMode.cpp
+++ b/Engine/Plugins/Editor/ModelingToolsEditorMode/Source/ModelingToolsEditorMode/Private/ModelingToolsEditorMode.cpp
@@ -3,6 +3,7 @@
 #include "ModelingToolsEditorMode.h"
 #include "Components/StaticMeshComponent.h"
 #include "IGeometryProcessingInterfacesModule.h"
+#include "Kismet2/ComponentEditorUtils.h"
 #include "Editor/UnrealEdEngine.h"
 #include "Tools/EdModeInteractiveToolsContext.h"
 #include "GeometryProcessingInterfaces/IUVEditorModularFeature.h"
@@ -171,6 +172,12 @@ FString GetToolName(const UInteractiveTool& Tool)
 	const FString* ToolName = FTextInspector::GetSourceString(Tool.GetToolInfo().ToolDisplayName);
 	return ToolName ? *ToolName : FString(TEXT("<Invalid ToolName>"));
 }
+
+bool bAllowNonEditableTargets = false;
+static FAutoConsoleVariableRef CVarAllowNonEditableComponents(
+	TEXT("ModelingTools.AllowNonEditableTargets"),
+	bAllowNonEditableTargets,
+	TEXT("Determines whether Modeling Mode allows tools to be run on internal components that do not show up in the component outliner."));
 }
 
 UModelingToolsEditorMode::UModelingToolsEditorMode()
@@ -308,6 +315,26 @@ void UModelingToolsEditorMode::Enter()
 	// Register builders for the generic component tool target, to support tools that only need the primitive component (e.g. the transform tool)
 	GetInteractiveToolsContext()->TargetManager->AddTargetFactory(NewObject<UPrimitiveComponentToolTargetFactory>(GetToolManager()));
 
+	GetInteractiveToolsContext()->TargetManager->SetInputFilterFunction([](const UObject* InputObject)->bool 
+	{
+		if (bAllowNonEditableTargets)
+		{
+			return true;
+		}
+			
+		// This check for editability is pretty much the same as the filtering that happens for component outliner in
+		//  ShouldAddInstancedActorComponent inside SSCSEditor::BuildSubTreeForActorNode
+		const UActorComponent* ActorComp = Cast<UActorComponent>(InputObject);
+		const USceneComponent* SceneComp = Cast<USceneComponent>(ActorComp);
+		const USceneComponent* ParentSceneComp = SceneComp != nullptr ? SceneComp->GetAttachParent() : nullptr;
+		if (FComponentEditorUtils::CanEditComponentInstance(ActorComp, ParentSceneComp,/*bAllowUserContructionScript*/ true))
+		{
+			return true;
+		}
+		return false;
+	});
+
+
 	// listen to post-build
 	GetToolManager()->OnToolPostBuild.AddUObject(this, &UModelingToolsEditorMode::OnToolPostBuild);
 
diff --git a/Engine/Source/Runtime/InteractiveToolsFramework/Private/ToolTargetManager.cpp b/Engine/Source/Runtime/InteractiveToolsFramework/Private/ToolTargetManager.cpp
--- a/Engine/Source/Runtime/InteractiveToolsFramework/Private/ToolTargetManager.cpp
+++ b/Engine/Source/Runtime/InteractiveToolsFramework/Private/ToolTargetManager.cpp
@@ -42,6 +42,10 @@ UToolTargetFactory* UToolTargetManager::FindFirstFactoryByPredicate(TFunctionRef
 
 bool UToolTargetManager::CanBuildTarget(UObject* SourceObject, const FToolTargetTypeRequirements& TargetType) const
 {
+	if (InputFilterFunction && !InputFilterFunction(SourceObject))
+	{
+		return false;
+	}
 	for (UToolTargetFactory* Factory : Factories)
 	{
 		if (Factory->CanBuildTarget(SourceObject, TargetType))
@@ -54,6 +58,10 @@ bool UToolTargetManager::CanBuildTarget(UObject* SourceObject, const FToolTarget
 
 UToolTarget* UToolTargetManager::BuildTarget(UObject* SourceObject, const FToolTargetTypeRequirements& TargetType)
 {
+	if (InputFilterFunction && !InputFilterFunction(SourceObject))
+	{
+		return nullptr;
+	}
 	for (UToolTargetFactory* Factory : Factories)
 	{
 		if (Factory->CanBuildTarget(SourceObject, TargetType))
diff --git a/Engine/Source/Runtime/InteractiveToolsFramework/Public/ToolTargetManager.h b/Engine/Source/Runtime/InteractiveToolsFramework/Public/ToolTargetManager.h
--- a/Engine/Source/Runtime/InteractiveToolsFramework/Public/ToolTargetManager.h
+++ b/Engine/Source/Runtime/InteractiveToolsFramework/Public/ToolTargetManager.h
@@ -130,6 +130,12 @@ public:
 	/** Shutdown the ToolTargetManager. Called by UInteractiveToolsContext. */
 	INTERACTIVETOOLSFRAMEWORK_API virtual void Shutdown();
 
+	/** Additional filter function that will be applied in CanBuildTarget and BuildTarget. */
+	INTERACTIVETOOLSFRAMEWORK_API virtual void SetInputFilterFunction(TFunction<bool(const UObject*)> ShouldAllowBuildingFromInputObject) 
+	{ 
+		InputFilterFunction = ShouldAllowBuildingFromInputObject; 
+	}
+
 protected:
 
 	// This should be removed once tools are transitioned to using tool targets, and the
@@ -146,4 +152,7 @@ protected:
 	TArray<TObjectPtr<UToolTargetFactory>> Factories;
 
 	// More state will go here if the manager deals with target caching.
+	
+private:
+	TFunction<bool(const UObject*)> InputFilterFunction;
 };
