ModelingTools: Fix spacing and flow not showing up at proper times.

diff --git a/Engine/Plugins/Editor/ModelingToolsEditorMode/Source/ModelingToolsEditorMode/Private/DetailsCustomizations/MeshVertexSculptToolCustomizations.cpp b/Engine/Plugins/Editor/ModelingToolsEditorMode/Source/ModelingToolsEditorMode/Private/DetailsCustomizations/MeshVertexSculptToolCustomizations.cpp
--- a/Engine/Plugins/Editor/ModelingToolsEditorMode/Source/ModelingToolsEditorMode/Private/DetailsCustomizations/MeshVertexSculptToolCustomizations.cpp
+++ b/Engine/Plugins/Editor/ModelingToolsEditorMode/Source/ModelingToolsEditorMode/Private/DetailsCustomizations/MeshVertexSculptToolCustomizations.cpp
@@ -50,49 +50,54 @@ void FSculptBrushPropertiesDetails::CustomizeDetails(IDetailLayoutBuilder& Detai
 	check(ObjectsBeingCustomized.Num() > 0);
 	USculptBrushProperties* BrushProperties = Cast<USculptBrushProperties>(ObjectsBeingCustomized[0]);
 
-	TSharedPtr<IPropertyHandle> FlowRateHandle = DetailBuilder.GetProperty(GET_MEMBER_NAME_CHECKED(USculptBrushProperties, FlowRate), USculptBrushProperties::StaticClass());
-	ensure(FlowRateHandle->IsValidHandle());
-
-	TSharedPtr<IPropertyHandle> SpacingHandle = DetailBuilder.GetProperty(GET_MEMBER_NAME_CHECKED(USculptBrushProperties, Spacing), USculptBrushProperties::StaticClass());
-	ensure(SpacingHandle->IsValidHandle());
-	SpacingHandle->MarkHiddenByCustomization();
-
-	TSharedPtr<IPropertyHandle> FalloffHandle = DetailBuilder.GetProperty(GET_MEMBER_NAME_CHECKED(USculptBrushProperties, BrushFalloffAmount), USculptBrushProperties::StaticClass());
-	ensure(FalloffHandle->IsValidHandle());
-	FalloffHandle->MarkHiddenByCustomization();
-
-	TSharedPtr<SDynamicNumericEntry::FDataSource> FlowRateSource = SDynamicNumericEntry::MakeSimpleDataSource(
-		FlowRateHandle, TInterval<float>(0.0f, 1.0f), TInterval<float>(0.0f, 1.0f));
-
-	if ( BrushProperties->bShowSpacing )
-	{
-		TSharedPtr<SDynamicNumericEntry::FDataSource> SpacingSource = SDynamicNumericEntry::MakeSimpleDataSource(
-			SpacingHandle, TInterval<float>(0.0f, 1000.0f), TInterval<float>(0.0f, 4.0f));
-
-		DetailBuilder.EditDefaultProperty(FlowRateHandle)->CustomWidget()
-			.OverrideResetToDefault(FResetToDefaultOverride::Hide())
-			.WholeRowContent()
-		[
-			MakeTwoWidgetDetailRowHBox(
-				MakeFixedWidthLabelSliderHBox(FlowRateHandle, FlowRateSource, FSculptToolsUIConstants::SculptShortLabelWidth),
-				MakeFixedWidthLabelSliderHBox(SpacingHandle, SpacingSource, FSculptToolsUIConstants::SculptShortLabelWidth)
-			)
-		];
-	}
-	else   // if bShowFalloff
+	// If we show falloff, it gets placed on the same line as flow/spacing. Otherwise, flow/spacing will get their own, uncustomized, lines
+	if (BrushProperties->bShowFalloff && (BrushProperties->bShowSpacing || BrushProperties->bShowFlowRate))
 	{
+		TSharedPtr<IPropertyHandle> FalloffHandle = DetailBuilder.GetProperty(GET_MEMBER_NAME_CHECKED(USculptBrushProperties, BrushFalloffAmount), USculptBrushProperties::StaticClass());
+		ensure(FalloffHandle->IsValidHandle());
+		FalloffHandle->MarkHiddenByCustomization();
 		TSharedPtr<SDynamicNumericEntry::FDataSource> FalloffSource = SDynamicNumericEntry::MakeSimpleDataSource(
 			FalloffHandle, TInterval<float>(0.0f, 1.0f), TInterval<float>(0.0f, 1.0f));
 
-		DetailBuilder.EditDefaultProperty(FlowRateHandle)->CustomWidget()
-			.OverrideResetToDefault(FResetToDefaultOverride::Hide())
-			.WholeRowContent()
+		// Note: the below assumes that edit conditions are set up in such a way that spacing and flow rate are mutually exclusive,
+		//  otherwise falloff (if shown) would show up twice. However, note that bShowSpacing and bShowFlowRate are NOT mutually
+		//  exclusive. We customize each one, and then their individual visibility controls the visibility of our customized row. 
+		if (BrushProperties->bShowSpacing)
+		{
+			TSharedPtr<IPropertyHandle> SpacingHandle = DetailBuilder.GetProperty(GET_MEMBER_NAME_CHECKED(USculptBrushProperties, Spacing), USculptBrushProperties::StaticClass());
+			ensure(SpacingHandle->IsValidHandle());
+			
+			TSharedPtr<SDynamicNumericEntry::FDataSource> SpacingSource = SDynamicNumericEntry::MakeSimpleDataSource(
+			SpacingHandle, TInterval<float>(0.0f, 1000.0f), TInterval<float>(0.0f, 4.0f));
+
+			DetailBuilder.EditDefaultProperty(SpacingHandle)->CustomWidget()
+				.OverrideResetToDefault(FResetToDefaultOverride::Hide())
+				.WholeRowContent()
+			[
+				MakeTwoWidgetDetailRowHBox(
+					MakeFixedWidthLabelSliderHBox(FalloffHandle, FalloffSource, FSculptToolsUIConstants::SculptShortLabelWidth),
+					MakeFixedWidthLabelSliderHBox(SpacingHandle, SpacingSource, FSculptToolsUIConstants::SculptShortLabelWidth)
+				)
+			];
+		}
+		if (BrushProperties->bShowFlowRate)
+		{
+			TSharedPtr<IPropertyHandle> FlowRateHandle = DetailBuilder.GetProperty(GET_MEMBER_NAME_CHECKED(USculptBrushProperties, FlowRate), USculptBrushProperties::StaticClass());
+			ensure(FlowRateHandle->IsValidHandle());
+		
+			TSharedPtr<SDynamicNumericEntry::FDataSource> FlowRateSource = SDynamicNumericEntry::MakeSimpleDataSource(
+				FlowRateHandle, TInterval<float>(0.0f, 1.0f), TInterval<float>(0.0f, 1.0f));
+
+			DetailBuilder.EditDefaultProperty(FlowRateHandle)->CustomWidget()
+				.OverrideResetToDefault(FResetToDefaultOverride::Hide())
+				.WholeRowContent()
 			[
 				MakeTwoWidgetDetailRowHBox(
 					MakeFixedWidthLabelSliderHBox(FalloffHandle, FalloffSource, FSculptToolsUIConstants::SculptShortLabelWidth),
 					MakeFixedWidthLabelSliderHBox(FlowRateHandle, FlowRateSource, FSculptToolsUIConstants::SculptShortLabelWidth)
 				)
 			];
+		}
 	}
 
 	TSharedPtr<IPropertyHandle> LazynessHandle = DetailBuilder.GetProperty(GET_MEMBER_NAME_CHECKED(USculptBrushProperties, Lazyness), USculptBrushProperties::StaticClass());
