Modeling Mode: update pressure sensitivity icon

diff --git a/Engine/Plugins/Editor/ModelingToolsEditorMode/Source/ModelingToolsEditorMode/Private/DetailsCustomizations/MeshVertexPaintToolCustomizations.cpp b/Engine/Plugins/Editor/ModelingToolsEditorMode/Source/ModelingToolsEditorMode/Private/DetailsCustomizations/MeshVertexPaintToolCustomizations.cpp
--- a/Engine/Plugins/Editor/ModelingToolsEditorMode/Source/ModelingToolsEditorMode/Private/DetailsCustomizations/MeshVertexPaintToolCustomizations.cpp
+++ b/Engine/Plugins/Editor/ModelingToolsEditorMode/Source/ModelingToolsEditorMode/Private/DetailsCustomizations/MeshVertexPaintToolCustomizations.cpp
@@ -86,7 +86,7 @@ void FVertexPaintBasicPropertiesDetails::BuildPaintPressureWidget(IDetailLayoutB
 			.AutoWidth()
 			[
 				SNew(SImage)
-				.Image(FModelingToolsEditorModeStyle::Get()->GetBrush("BrushIcons.PressureSensitivity"))
+				.Image(FModelingToolsEditorModeStyle::Get()->GetBrush("BrushIcons.PressureSensitivity.Small"))
 			]
 		]
 	];
diff --git a/Engine/Plugins/Editor/ModelingToolsEditorMode/Source/ModelingToolsEditorMode/Private/DetailsCustomizations/MeshVertexSculptToolCustomizations.cpp b/Engine/Plugins/Editor/ModelingToolsEditorMode/Source/ModelingToolsEditorMode/Private/DetailsCustomizations/MeshVertexSculptToolCustomizations.cpp
--- a/Engine/Plugins/Editor/ModelingToolsEditorMode/Source/ModelingToolsEditorMode/Private/DetailsCustomizations/MeshVertexSculptToolCustomizations.cpp
+++ b/Engine/Plugins/Editor/ModelingToolsEditorMode/Source/ModelingToolsEditorMode/Private/DetailsCustomizations/MeshVertexSculptToolCustomizations.cpp
@@ -244,7 +244,7 @@ void FMeshSculptBrushOpPropertiesDetails<BrushClass>::BuildStrengthPressureWidge
 			.AutoWidth()
 			[
 				SNew(SImage)
-				.Image(FModelingToolsEditorModeStyle::Get()->GetBrush("BrushIcons.PressureSensitivity"))
+				.Image(FModelingToolsEditorModeStyle::Get()->GetBrush("BrushIcons.PressureSensitivity.Small"))
 			]
 		]
 	];
diff --git a/Engine/Plugins/Editor/ModelingToolsEditorMode/Source/ModelingToolsEditorMode/Private/DetailsCustomizations/ModelingToolsBrushSizeCustomization.cpp b/Engine/Plugins/Editor/ModelingToolsEditorMode/Source/ModelingToolsEditorMode/Private/DetailsCustomizations/ModelingToolsBrushSizeCustomization.cpp
--- a/Engine/Plugins/Editor/ModelingToolsEditorMode/Source/ModelingToolsEditorMode/Private/DetailsCustomizations/ModelingToolsBrushSizeCustomization.cpp
+++ b/Engine/Plugins/Editor/ModelingToolsEditorMode/Source/ModelingToolsEditorMode/Private/DetailsCustomizations/ModelingToolsBrushSizeCustomization.cpp
@@ -173,7 +173,7 @@ void FModelingToolsBrushSizeCustomization::CustomizeHeader(TSharedRef<IPropertyH
 						.AutoWidth()
 						[
 							SNew(SImage)
-							.Image(FModelingToolsEditorModeStyle::Get()->GetBrush("BrushIcons.PressureSensitivity"))
+							.Image(FModelingToolsEditorModeStyle::Get()->GetBrush("BrushIcons.PressureSensitivity.Small"))
 						]
 					]
 			]
diff --git a/Engine/Plugins/Editor/ModelingToolsEditorMode/Source/ModelingToolsEditorMode/Private/ModelingToolsEditorModeStyle.cpp b/Engine/Plugins/Editor/ModelingToolsEditorMode/Source/ModelingToolsEditorMode/Private/ModelingToolsEditorModeStyle.cpp
--- a/Engine/Plugins/Editor/ModelingToolsEditorMode/Source/ModelingToolsEditorMode/Private/ModelingToolsEditorModeStyle.cpp
+++ b/Engine/Plugins/Editor/ModelingToolsEditorMode/Source/ModelingToolsEditorMode/Private/ModelingToolsEditorModeStyle.cpp
@@ -475,7 +475,8 @@ void FModelingToolsEditorModeStyle::Initialize()
 		StyleSet->Set("BrushTypeIcons.HeightErase", new FSlateImageBrush(StyleSet->RootToCoreContentDir(TEXT("../Editor/Slate/Icons/Landscape/Landscape_Erase_x40.png")), Icon20x20));
 
 		// Pressure Sensitivity toggle button
-		StyleSet->Set("BrushIcons.PressureSensitivity", new IMAGE_BRUSH_SVG("Icons/PressureSensitivity", Icon16x16));
+		StyleSet->Set("BrushIcons.PressureSensitivity", new IMAGE_BRUSH_SVG("Icons/PressureSensitivity", Icon20x20));
+		StyleSet->Set("BrushIcons.PressureSensitivity.Small", new IMAGE_BRUSH_SVG("Icons/PressureSensitivity", Icon16x16));
 
 		//
 		// Icons for selection buttons in PolyEd and TriEd
diff --git a/Engine/Plugins/Runtime/MeshModelingToolset/Source/MeshModelingTools/Private/Sculpting/MeshSculptToolBase.cpp b/Engine/Plugins/Runtime/MeshModelingToolset/Source/MeshModelingTools/Private/Sculpting/MeshSculptToolBase.cpp
--- a/Engine/Plugins/Runtime/MeshModelingToolset/Source/MeshModelingTools/Private/Sculpting/MeshSculptToolBase.cpp
+++ b/Engine/Plugins/Runtime/MeshModelingToolset/Source/MeshModelingTools/Private/Sculpting/MeshSculptToolBase.cpp
@@ -874,6 +874,15 @@ bool UMeshSculptToolBase::GetBrushStrengthPressureEnabled()
 	return false;
 }
 
+void UMeshSculptToolBase::SetBrushStrengthPressureEnabled(const bool bEnabled)
+{
+	TUniquePtr<FMeshSculptBrushOp>& BrushOp = GetActiveBrushOp();
+	if (BrushOp->PropertySet.IsValid())
+	{
+		BrushOp->PropertySet->bIsStrengthPressureEnabled = bEnabled;
+	}
+}
+
 void UMeshSculptToolBase::UpdateBrushFrameWorld(const FVector3d& NewPosition, const FVector3d& NewNormal)
 {
 	FFrame3d PrevBrushFrameWorld = LastBrushFrameWorld;
diff --git a/Engine/Plugins/Runtime/MeshModelingToolset/Source/MeshModelingTools/Public/Sculpting/MeshSculptToolBase.h b/Engine/Plugins/Runtime/MeshModelingToolset/Source/MeshModelingTools/Public/Sculpting/MeshSculptToolBase.h
--- a/Engine/Plugins/Runtime/MeshModelingToolset/Source/MeshModelingTools/Public/Sculpting/MeshSculptToolBase.h
+++ b/Engine/Plugins/Runtime/MeshModelingToolset/Source/MeshModelingTools/Public/Sculpting/MeshSculptToolBase.h
@@ -297,9 +297,13 @@ protected:
 	/** Returns true if BrushSize Pressure Sensitivity is both supported by the tool and currently enabled. */
 	UE_API bool GetBrushSizePressureSensitivityEnabled() const;
 
+public:
 	/** Returns true if Pressure Sensitivity for Brush Strength is both supported by the tool and currently enabled. */
 	UE_API bool GetBrushStrengthPressureEnabled();
+	/** Sets the Pressure Sensitivity for Brush Strength to be enabled/disabled */
+	UE_API void SetBrushStrengthPressureEnabled(const bool bEnabled);
 
+protected:
 	/**
 	 * Subclass must implement this and return relevant rendering component
 	 */
