ModelingTools: Fixed a nullptr dereference crash in OnActiveViewportChanged for StylusHandling

diff --git a/Engine/Plugins/Editor/ModelingToolsEditorMode/Source/ModelingToolsEditorMode/Private/ModelingToolsEditorModeToolkit.cpp b/Engine/Plugins/Editor/ModelingToolsEditorMode/Source/ModelingToolsEditorMode/Private/ModelingToolsEditorModeToolkit.cpp
--- a/Engine/Plugins/Editor/ModelingToolsEditorMode/Source/ModelingToolsEditorMode/Private/ModelingToolsEditorModeToolkit.cpp
+++ b/Engine/Plugins/Editor/ModelingToolsEditorMode/Source/ModelingToolsEditorMode/Private/ModelingToolsEditorModeToolkit.cpp
@@ -2427,7 +2427,7 @@ void FModelingToolsEditorModeToolkit::OnActiveViewportChanged(TSharedPtr<IAssetV
 	}
 
 #if ENABLE_STYLUS_SUPPORT
-	if (StylusInputHandler)
+	if (StylusInputHandler && NewViewport.IsValid())
 	{
 		StylusInputHandler->RegisterWindow(NewViewport->AsWidget());
 	}
