fix undo/redo not updating the sculpt layers info in vertex sculpt

#Synced-CL 51242673

diff --git a/Engine/Plugins/Editor/ModelingToolsEditorMode/Source/ModelingToolsEditorMode/Private/DetailsCustomizations/MeshVertexSculptToolCustomizations.cpp b/Engine/Plugins/Editor/ModelingToolsEditorMode/Source/ModelingToolsEditorMode/Private/DetailsCustomizations/MeshVertexSculptToolCustomizations.cpp
--- a/Engine/Plugins/Editor/ModelingToolsEditorMode/Source/ModelingToolsEditorMode/Private/DetailsCustomizations/MeshVertexSculptToolCustomizations.cpp
+++ b/Engine/Plugins/Editor/ModelingToolsEditorMode/Source/ModelingToolsEditorMode/Private/DetailsCustomizations/MeshVertexSculptToolCustomizations.cpp
@@ -128,6 +128,14 @@ FMeshSculptLayerPropertiesDetails::FMeshSculptLayerPropertiesDetails()
 	Controller = MakeShared<FSculptLayersController>();
 }
 
+FMeshSculptLayerPropertiesDetails::~FMeshSculptLayerPropertiesDetails()
+{
+	if (SculptProperties.IsValid())
+	{
+		SculptProperties->OnMeshLayersPanelRequestRebuild.Remove(SculptUpdatesHandle);
+	}
+}
+
 TSharedRef<IDetailCustomization> FMeshSculptLayerPropertiesDetails::MakeInstance()
 {
 	return MakeShareable(new FMeshSculptLayerPropertiesDetails);
@@ -139,7 +147,7 @@ void FMeshSculptLayerPropertiesDetails::CustomizeDetails(IDetailLayoutBuilder& D
 	DetailBuilder.GetObjectsBeingCustomized(ObjectsBeingCustomized);
 	check(ObjectsBeingCustomized.Num() > 0);
 	UMeshSculptLayerProperties* SculptLayersProperties = Cast<UMeshSculptLayerProperties>(ObjectsBeingCustomized[0]);
-
+	SculptProperties = SculptLayersProperties;
 	Controller->SetProperties(SculptLayersProperties);
 
 	TSharedPtr<IPropertyHandle> ActiveLayerHandle = DetailBuilder.GetProperty(GET_MEMBER_NAME_CHECKED(UMeshSculptLayerProperties, ActiveLayer), UMeshSculptLayerProperties::StaticClass());
@@ -156,6 +164,13 @@ void FMeshSculptLayerPropertiesDetails::CustomizeDetails(IDetailLayoutBuilder& D
 	[
 		SNew(SMeshLayersStack, Controller)
 	];
+
+	// verify invariant that this is always a fresh details object (if this were not the case, would need to unbind delegate handle above.)
+	ensureMsgf(!SculptUpdatesHandle.IsValid(), TEXT("Unexpected re-use of detail customization object"));
+	// ensure this stack is rebuilt when Mesh Layers are updated elsewhere in the editor
+	SculptUpdatesHandle = SculptLayersProperties->OnMeshLayersPanelRequestRebuild.AddLambda([this]() {
+		Controller->RefreshLayersStackView();
+	});
 }
 
 template<typename BrushClass>
diff --git a/Engine/Plugins/Editor/ModelingToolsEditorMode/Source/ModelingToolsEditorMode/Public/DetailsCustomizations/MeshVertexSculptToolCustomizations.h b/Engine/Plugins/Editor/ModelingToolsEditorMode/Source/ModelingToolsEditorMode/Public/DetailsCustomizations/MeshVertexSculptToolCustomizations.h
--- a/Engine/Plugins/Editor/ModelingToolsEditorMode/Source/ModelingToolsEditorMode/Public/DetailsCustomizations/MeshVertexSculptToolCustomizations.h
+++ b/Engine/Plugins/Editor/ModelingToolsEditorMode/Source/ModelingToolsEditorMode/Public/DetailsCustomizations/MeshVertexSculptToolCustomizations.h
@@ -21,6 +21,7 @@ class UMeshVertexSculptTool;
 class FRecentAlphasProvider;
 class SToolInputAssetComboPanel;
 class IMeshLayersController;
+class UMeshSculptLayerProperties;
 
 
 // customization for USculptBrushProperties, creates two-column layout
@@ -103,10 +104,13 @@ class FMeshSculptLayerPropertiesDetails : public IDetailCustomization
 {
 public:
 	FMeshSculptLayerPropertiesDetails();
+	~FMeshSculptLayerPropertiesDetails();
 	
 	static TSharedRef<IDetailCustomization> MakeInstance();
 	void CustomizeDetails(IDetailLayoutBuilder& DetailBuilder) override;
 	
 private:
 	TSharedPtr<IMeshLayersController> Controller;
+	FDelegateHandle SculptUpdatesHandle;
+	TWeakObjectPtr<UMeshSculptLayerProperties> SculptProperties;
 };
\ No newline at end of file
diff --git a/Engine/Plugins/Runtime/MeshModelingToolset/Source/MeshModelingTools/Private/MeshVertexSculptTool.cpp b/Engine/Plugins/Runtime/MeshModelingToolset/Source/MeshModelingTools/Private/MeshVertexSculptTool.cpp
--- a/Engine/Plugins/Runtime/MeshModelingToolset/Source/MeshModelingTools/Private/MeshVertexSculptTool.cpp
+++ b/Engine/Plugins/Runtime/MeshModelingToolset/Source/MeshModelingTools/Private/MeshVertexSculptTool.cpp
@@ -490,6 +490,11 @@ void UMeshVertexSculptTool::UpdateToolMeshes(TFunctionRef<TUniquePtr<FMeshRegion
 			{
 				this->WaitForPendingUndoRedoUpdate();
 			};
+			// refresh sculpt layer UI w/ any external mesh changes
+			NewChange->AfterModify = [this](bool bRevert)
+			{
+				this->RefreshSculptLayerProperties();
+			};
 			// Note this change should be in the context of a larger transaction, so the text isn't that important
 			GetToolManager()->EmitObjectChange(DynamicMeshComponent, MoveTemp(NewChange), LOCTEXT("UpdateVertexSculptMesh", "Updated Mesh"));
 			
@@ -1936,10 +1941,17 @@ void UMeshVertexSculptTool::OnDynamicMeshComponentChanged(UDynamicMeshComponent*
 		bUndoUpdatePending = false;
 	});
 
+	RefreshSculptLayerProperties();
 }
 
 
-
+void UMeshVertexSculptTool::RefreshSculptLayerProperties()
+{
+	if (SculptLayerProperties)
+	{
+		SculptLayerProperties->RefreshFromMesh();
+	}
+}
 
 
 
diff --git a/Engine/Plugins/Runtime/MeshModelingToolset/Source/MeshModelingTools/Private/Properties/MeshSculptLayerProperties.cpp b/Engine/Plugins/Runtime/MeshModelingToolset/Source/MeshModelingTools/Private/Properties/MeshSculptLayerProperties.cpp
--- a/Engine/Plugins/Runtime/MeshModelingToolset/Source/MeshModelingTools/Private/Properties/MeshSculptLayerProperties.cpp
+++ b/Engine/Plugins/Runtime/MeshModelingToolset/Source/MeshModelingTools/Private/Properties/MeshSculptLayerProperties.cpp
@@ -287,6 +287,18 @@ FName UMeshSculptLayerProperties::GetLayerName(const int32 InIndex)
 	return LayerName;
 }
 
+void UMeshSculptLayerProperties::RefreshFromMesh()
+{
+	ReadSculptLayers([this](const UE::Geometry::FDynamicMesh3& Mesh, const UE::Geometry::FDynamicMeshSculptLayers* SculptLayers)->void
+	{
+		if (SculptLayers)
+		{
+			UpdateSettingsFromMesh(SculptLayers);
+			OnMeshLayersPanelRequestRebuild.Broadcast();
+		}
+	});
+}
+
 void UMeshSculptLayerProperties::SetLayerName(const int32 InIndex, const FName InLayerName)
 {
 	// ensure the provided index is valid
diff --git a/Engine/Plugins/Runtime/MeshModelingToolset/Source/MeshModelingTools/Public/MeshVertexSculptTool.h b/Engine/Plugins/Runtime/MeshModelingToolset/Source/MeshModelingTools/Public/MeshVertexSculptTool.h
--- a/Engine/Plugins/Runtime/MeshModelingToolset/Source/MeshModelingTools/Public/MeshVertexSculptTool.h
+++ b/Engine/Plugins/Runtime/MeshModelingToolset/Source/MeshModelingTools/Public/MeshVertexSculptTool.h
@@ -285,6 +285,9 @@ private:
 	UPROPERTY()
 	TObjectPtr<UMeshSculptLayerProperties> SculptLayerProperties;
 
+	// Call to sync sculpt layer properties to the current mesh state
+	void RefreshSculptLayerProperties();
+
 public:
 	UE_API virtual void IncreaseBrushSpeedAction() override;
 	UE_API virtual void DecreaseBrushSpeedAction() override;
diff --git a/Engine/Plugins/Runtime/MeshModelingToolset/Source/MeshModelingTools/Public/Properties/MeshSculptLayerProperties.h b/Engine/Plugins/Runtime/MeshModelingToolset/Source/MeshModelingTools/Public/Properties/MeshSculptLayerProperties.h
--- a/Engine/Plugins/Runtime/MeshModelingToolset/Source/MeshModelingTools/Public/Properties/MeshSculptLayerProperties.h
+++ b/Engine/Plugins/Runtime/MeshModelingToolset/Source/MeshModelingTools/Public/Properties/MeshSculptLayerProperties.h
@@ -61,12 +61,19 @@ public:
 
 	MESHMODELINGTOOLS_API void Init(IModelingToolExternalDynamicMeshUpdateAPI* InTool, int32 InNumLockedBaseLayers);
 
+	/** Refresh properties from current mesh data */
+	MESHMODELINGTOOLS_API void RefreshFromMesh();
+
 #if WITH_EDITOR
 	MESHMODELINGTOOLS_API virtual void PostEditChangeProperty(FPropertyChangedEvent& PropertyChangedEvent) override;
 #endif
 
 	int32 GetNumLockedBaseLayers() const { return NumLockedBaseLayers; }
 
+
+	DECLARE_MULTICAST_DELEGATE(OnMeshLayersPropertiesPanelRequestRebuild);
+	OnMeshLayersPropertiesPanelRequestRebuild OnMeshLayersPanelRequestRebuild;
+
 private:
 
 	IModelingToolExternalDynamicMeshUpdateAPI* Tool = nullptr;
