[ModelingTools Exp] Extract most of data flow vert paint tool into a sharable base class

#Synced-CL 51489786

diff --git a/Engine/Plugins/Editor/ModelingToolsEditorMode/Source/ModelingToolsEditorMode/Private/ModelingToolsEditorMode.cpp b/Engine/Plugins/Editor/ModelingToolsEditorMode/Source/ModelingToolsEditorMode/Private/ModelingToolsEditorMode.cpp
--- a/Engine/Plugins/Editor/ModelingToolsEditorMode/Source/ModelingToolsEditorMode/Private/ModelingToolsEditorMode.cpp
+++ b/Engine/Plugins/Editor/ModelingToolsEditorMode/Source/ModelingToolsEditorMode/Private/ModelingToolsEditorMode.cpp
@@ -138,6 +138,7 @@
 #include "Selection/StaticMeshSelector.h"
 #include "ModelingSelectionInteraction.h"
 #include "DynamicMeshActor.h"
+#include "MeshAttributePaintToolV2.h"
 #include "ModelingToolsEditorModeStyle.h"
 #include "Components/BrushComponent.h"
 #include "Engine/StaticMeshActor.h"
@@ -179,8 +180,16 @@ static FAutoConsoleVariableRef CVarAllowNonEditableComponents(
 	TEXT("ModelingTools.AllowNonEditableTargets"),
 	bAllowNonEditableTargets,
 	TEXT("Determines whether Modeling Mode allows tools to be run on internal components that do not show up in the component outliner."));
+
+	
+bool bUseExperimentalMeshAttributePaintTool = true;
+static FAutoConsoleVariableRef CVarUseExperimentalMeshAttributePaintTool(
+	TEXT("ModelingTools.UseExperimentalMeshAttributePaintTool"),
+	bUseExperimentalMeshAttributePaintTool,
+	TEXT("Determines whether to use the new experimental mesh attribute paint tool."));
 }
 
+
 UModelingToolsEditorMode::UModelingToolsEditorMode()
 {
 	Info = FEditorModeInfo(
@@ -737,8 +746,15 @@ void UModelingToolsEditorMode::Enter()
 
 	auto EditMeshMaterialsToolBuilder = NewObject<UEditMeshMaterialsToolBuilder>();
 	RegisterTool(ToolManagerCommands.BeginEditMeshMaterialsTool, TEXT("BeginEditMeshMaterialsTool"), EditMeshMaterialsToolBuilder);
-	
-	RegisterTool(ToolManagerCommands.BeginMeshAttributePaintTool, TEXT("BeginMeshAttributePaintTool"), NewObject<UMeshAttributePaintToolBuilder>());
+
+	if (bUseExperimentalMeshAttributePaintTool)
+	{
+		RegisterTool(ToolManagerCommands.BeginMeshAttributePaintTool, TEXT("BeginMeshAttributePaintTool"), NewObject<UMeshAttributePaintToolV2Builder>());
+	}
+	else
+	{
+		RegisterTool(ToolManagerCommands.BeginMeshAttributePaintTool, TEXT("BeginMeshAttributePaintTool"), NewObject<UMeshAttributePaintToolBuilder>());
+	}
 
 	auto BakeMeshAttributeMapsToolBuilder = NewObject<UBakeMeshAttributeMapsToolBuilder>();
 	RegisterTool(ToolManagerCommands.BeginBakeMeshAttributeMapsTool, TEXT("BeginBakeMeshAttributeMapsTool"), BakeMeshAttributeMapsToolBuilder);
diff --git a/Engine/Plugins/Editor/ModelingToolsEditorMode/Source/ModelingToolsEditorMode/Private/ModelingToolsEditorModeModule.cpp b/Engine/Plugins/Editor/ModelingToolsEditorMode/Source/ModelingToolsEditorMode/Private/ModelingToolsEditorModeModule.cpp
--- a/Engine/Plugins/Editor/ModelingToolsEditorMode/Source/ModelingToolsEditorMode/Private/ModelingToolsEditorModeModule.cpp
+++ b/Engine/Plugins/Editor/ModelingToolsEditorMode/Source/ModelingToolsEditorMode/Private/ModelingToolsEditorModeModule.cpp
@@ -26,6 +26,8 @@
 #include "BakeMultiMeshAttributeMapsTool.h"
 #include "BakeMeshAttributeVertexTool.h"
 #include "BakeTransformTool.h"
+#include "MeshVertexAttributePaintToolBase.h"
+#include "DetailsCustomizations/MeshVertexAttributePaintToolDetailCustomization.h"
 #include "Sculpting/KelvinletBrushOp.h"
 #include "Sculpting/MeshInflateBrushOps.h"
 #include "Sculpting/MeshMoveBrushOps.h"
@@ -156,6 +158,10 @@ void FModelingToolsEditorModeModule::OnPostEngineInit()
 
 	PropertyModule.RegisterCustomClassLayout("BakeTransformToolProperties", FOnGetDetailCustomizationInstance::CreateStatic(&FBakeTransformToolDetails::MakeInstance));
 	ClassesToUnregisterOnShutdown.Add(UBakeTransformToolProperties::StaticClass()->GetFName());
+
+	// Vertex Attribute Paint
+	PropertyModule.RegisterCustomClassLayout(UMeshVertexAttributePaintToolProperties::StaticClass()->GetFName(), FOnGetDetailCustomizationInstance::CreateStatic(&UE::MeshVertexAttributePaintToolBase::FVertexAttributePaintToolDetailCustomization::MakeInstance));
+	ClassesToUnregisterOnShutdown.Add(UMeshVertexAttributePaintToolProperties::StaticClass()->GetFName());
 }
 
 #undef LOCTEXT_NAMESPACE
diff --git a/Engine/Plugins/Experimental/Dataflow/Source/DataflowEditor/Private/Dataflow/DataflowEditorModule.cpp b/Engine/Plugins/Experimental/Dataflow/Source/DataflowEditor/Private/Dataflow/DataflowEditorModule.cpp
--- a/Engine/Plugins/Experimental/Dataflow/Source/DataflowEditor/Private/Dataflow/DataflowEditorModule.cpp
+++ b/Engine/Plugins/Experimental/Dataflow/Source/DataflowEditor/Private/Dataflow/DataflowEditorModule.cpp
@@ -21,7 +21,6 @@
 #include "DataflowEditorTools/DataflowEditorWeightMapPaintTool.h"
 #include "DataflowEditorTools/DataflowEditorVertexAttributePaintTool.h"
 #include "Dataflow/DataflowCollectionAddScalarVertexPropertyNode.h"
-#include "DataflowEditorTools/DataflowEditorVertexAttributePaintToolDetailCustomization.h"
 
 #include "Dataflow/DataflowInstanceDetails.h"
 
@@ -313,7 +312,6 @@ void FDataflowEditorModule::StartupModule()
 
 		PropertyModule->RegisterCustomPropertyTypeLayout(DataflowColorRampName, FOnGetPropertyTypeCustomizationInstance::CreateStatic(&UE::Dataflow::FDataflowColorRampCustomization::MakeInstance));
 		PropertyModule->RegisterCustomPropertyTypeLayout(DataflowLinearColorRampName, FOnGetPropertyTypeCustomizationInstance::CreateStatic(&UE::Dataflow::FDataflowLinearColorRampCustomization::MakeInstance));
-		PropertyModule->RegisterCustomClassLayout(UDataflowEditorVertexAttributePaintToolProperties::StaticClass()->GetFName(), FOnGetDetailCustomizationInstance::CreateStatic(&UE::Dataflow::Editor::FVertexAttributePaintToolDetailCustomization::MakeInstance));
 
 		UE::Dataflow::RegisterAnyTypeCustomizations(*PropertyModule);
 	}
@@ -393,11 +391,6 @@ void FDataflowEditorModule::ShutdownModule()
 		PropertyModule->UnregisterCustomPropertyTypeLayout(DataflowColorRampName);
 		PropertyModule->UnregisterCustomPropertyTypeLayout(DataflowLinearColorRampName);
 
-		if (UObjectInitialized())
-		{
-			PropertyModule->UnregisterCustomClassLayout(UDataflowEditorVertexAttributePaintToolProperties::StaticClass()->GetFName());
-		}
-
 		UE::Dataflow::UnregisterAnyTypeCustomizations(*PropertyModule);
 	}
 
