Modeling tools : Mesh selection mechanic - add optional toggle buttons to toggle between marquee and lasso selection tools

#Synced-CL 51713790

diff --git a/Engine/Plugins/Editor/ModelingToolsEditorMode/Source/ModelingToolsEditorMode/Private/DetailsCustomizations/MeshTopologySelectionMechanicCustomization.cpp b/Engine/Plugins/Editor/ModelingToolsEditorMode/Source/ModelingToolsEditorMode/Private/DetailsCustomizations/MeshTopologySelectionMechanicCustomization.cpp
--- a/Engine/Plugins/Editor/ModelingToolsEditorMode/Source/ModelingToolsEditorMode/Private/DetailsCustomizations/MeshTopologySelectionMechanicCustomization.cpp
+++ b/Engine/Plugins/Editor/ModelingToolsEditorMode/Source/ModelingToolsEditorMode/Private/DetailsCustomizations/MeshTopologySelectionMechanicCustomization.cpp
@@ -45,7 +45,7 @@ void FMeshTopologySelectionMechanicPropertiesDetails::CustomizeDetails(IDetailLa
 		ToolbarBuilder.AddToolBarButton(FUIAction(
 			FExecuteAction::CreateWeakLambda(Properties, [BoolPropertyHandle]()
 			{
-				bool bCurrentValue;
+				bool bCurrentValue = false;
 				BoolPropertyHandle->GetValue(bCurrentValue);
 				BoolPropertyHandle->SetValue(!bCurrentValue);
 			}),
@@ -54,7 +54,7 @@ void FMeshTopologySelectionMechanicPropertiesDetails::CustomizeDetails(IDetailLa
 			}),
 			FIsActionChecked::CreateWeakLambda(Properties, [BoolPropertyHandle]()
 			{
-				bool bValue;
+				bool bValue = false;
 				BoolPropertyHandle->GetValue(bValue);
 				return bValue;
 			})),
@@ -65,6 +65,45 @@ void FMeshTopologySelectionMechanicPropertiesDetails::CustomizeDetails(IDetailLa
 			EUserInterfaceActionType::ToggleButton);
 	};
 
+	if (Properties->bDisplaySelectionDragToolsControls)
+	{
+		auto AddToggleButtonForSelectionDragToolEnum = [&ToolbarBuilder, Properties](TSharedPtr<IPropertyHandle> EnumPropertyHandle, EMeshTopologySelectionMechanicDragTool EnumValue, FSlateIcon Icon)
+			{
+				ToolbarBuilder.AddToolBarButton(FUIAction(
+					FExecuteAction::CreateWeakLambda(Properties, [EnumPropertyHandle, EnumValue]()
+						{
+							EnumPropertyHandle->SetValue((uint8)EnumValue);
+						}),
+					FCanExecuteAction::CreateWeakLambda(Properties, [EnumPropertyHandle]() {
+						return EnumPropertyHandle->IsEditable();
+						}),
+					FIsActionChecked::CreateWeakLambda(Properties, [EnumPropertyHandle, EnumValue]()
+						{
+							uint8 Value = 0;
+							EnumPropertyHandle->GetValue(Value);
+							return (Value == (uint8)EnumValue);
+						})),
+					NAME_None, // Extension hook
+					TAttribute<FText>(EnumPropertyHandle->GetPropertyDisplayName()), // Label
+					TAttribute<FText>(EnumPropertyHandle->GetToolTipText()), // Tooltip
+					TAttribute<FSlateIcon>(Icon),
+					EUserInterfaceActionType::ToggleButton);
+			};
+
+		TSharedPtr<IPropertyHandle> SelectionDragToolHandle = DetailBuilder.GetProperty(GET_MEMBER_NAME_CHECKED(UMeshTopologySelectionMechanicProperties, SelectionDragTool),
+			UMeshTopologySelectionMechanicProperties::StaticClass());
+		SelectionDragToolHandle->MarkHiddenByCustomization();
+
+		ToolbarBuilder.BeginSection("SelectionDragTools");
+		ToolbarBuilder.BeginBlockGroup();
+
+		AddToggleButtonForSelectionDragToolEnum(SelectionDragToolHandle, EMeshTopologySelectionMechanicDragTool::Marquee, FSlateIcon(FAppStyle::GetAppStyleSetName(), "LevelEditor.StrictBoxSelect"));
+		AddToggleButtonForSelectionDragToolEnum(SelectionDragToolHandle, EMeshTopologySelectionMechanicDragTool::Lasso, FSlateIcon(FModelingToolsEditorModeStyle::Get()->GetStyleSetName(), "PolyEd.SelectByLasso"));
+
+		ToolbarBuilder.EndBlockGroup();
+		ToolbarBuilder.EndSection();
+	}
+
 	ToolbarBuilder.BeginSection("SelectionFilter");
 	ToolbarBuilder.BeginBlockGroup();
 
diff --git a/Engine/Plugins/Editor/ModelingToolsEditorMode/Source/ModelingToolsEditorMode/Private/ModelingToolsEditorModeStyle.cpp b/Engine/Plugins/Editor/ModelingToolsEditorMode/Source/ModelingToolsEditorMode/Private/ModelingToolsEditorModeStyle.cpp
--- a/Engine/Plugins/Editor/ModelingToolsEditorMode/Source/ModelingToolsEditorMode/Private/ModelingToolsEditorModeStyle.cpp
+++ b/Engine/Plugins/Editor/ModelingToolsEditorMode/Source/ModelingToolsEditorMode/Private/ModelingToolsEditorModeStyle.cpp
@@ -482,6 +482,9 @@ void FModelingToolsEditorModeStyle::Initialize()
 		// Icons for selection buttons in PolyEd and TriEd
 		//
 
+		StyleSet->Set("PolyEd.SelectByMarquee", new FSlateVectorImageBrush(StyleSet->RootToCoreContentDir(TEXT("../Editor/Slate/Starship/Common/MarqueeSelectOccluded_16")), Icon20x20));
+		StyleSet->Set("PolyEd.SelectByLasso", new FSlateImageBrush(StyleSet->RootToCoreContentDir(TEXT("../Editor/Slate/Icons/GeneralTools/Lasso_40x.png")), Icon20x20));
+
 		StyleSet->Set("PolyEd.SelectCorners", new IMAGE_BRUSH_SVG("Icons/SelectionVertices", Icon20x20));
 		StyleSet->Set("PolyEd.SelectEdges", new IMAGE_BRUSH_SVG("Icons/SelectionBorderEdges", Icon20x20));
 		StyleSet->Set("PolyEd.SelectFaces", new IMAGE_BRUSH_SVG("Icons/SelectionTriangles3", Icon20x20));
diff --git a/Engine/Plugins/Runtime/MeshModelingToolset/Source/ModelingComponents/Public/Selection/MeshTopologySelectionMechanic.h b/Engine/Plugins/Runtime/MeshModelingToolset/Source/ModelingComponents/Public/Selection/MeshTopologySelectionMechanic.h
--- a/Engine/Plugins/Runtime/MeshModelingToolset/Source/ModelingComponents/Public/Selection/MeshTopologySelectionMechanic.h
+++ b/Engine/Plugins/Runtime/MeshModelingToolset/Source/ModelingComponents/Public/Selection/MeshTopologySelectionMechanic.h
@@ -99,6 +99,9 @@ public:
 	UPROPERTY(EditAnywhere, Category = "AdditionalSelectionOptions|Ortho Viewport Behavior")
 	bool bIgnoreOcclusion = false;
 
+	// When true display toggles buttons to allow switching between various selection drag tools (in the detail customization)
+	bool bDisplaySelectionDragToolsControls = false;
+
 	// Used to avoid showing some of the selection filter buttons in triedit (in the detail customization)
 	bool bDisplayPolygroupReliantControls = true;
 
