Modeling Tools - Reworded the vertex paint brush mode tooltips to accurately reflect the recent renaming of the Strength property to Value.

#rnx
#Synced-CL 51673482

diff --git a/Engine/Plugins/Editor/ModelingToolsEditorMode/Source/ModelingToolsEditorMode/Private/DetailsCustomizations/MeshVertexAttributePaintToolDetailCustomization.cpp b/Engine/Plugins/Editor/ModelingToolsEditorMode/Source/ModelingToolsEditorMode/Private/DetailsCustomizations/MeshVertexAttributePaintToolDetailCustomization.cpp
--- a/Engine/Plugins/Editor/ModelingToolsEditorMode/Source/ModelingToolsEditorMode/Private/DetailsCustomizations/MeshVertexAttributePaintToolDetailCustomization.cpp
+++ b/Engine/Plugins/Editor/ModelingToolsEditorMode/Source/ModelingToolsEditorMode/Private/DetailsCustomizations/MeshVertexAttributePaintToolDetailCustomization.cpp
@@ -674,16 +674,16 @@ namespace UE::MeshVertexAttributePaintToolBase
 					})
 				+ SSegmentedControl<EMeshVertexAttributePaintToolEditOperation>::Slot(EMeshVertexAttributePaintToolEditOperation::Add)
 				.Text(LOCTEXT("BrushAddMode", "Add"))
-				.ToolTip(LOCTEXT("BrushAddTooltip", "Add: Applies the current value plus the value to the new value."))
+				.ToolTip(LOCTEXT("BrushAddTooltip", "Add: Increases the existing vertex value by the brush value."))
 				+ SSegmentedControl<EMeshVertexAttributePaintToolEditOperation>::Slot(EMeshVertexAttributePaintToolEditOperation::Replace)
 				.Text(LOCTEXT("BrushReplaceMode", "Replace"))
-				.ToolTip(LOCTEXT("BrushReplaceTooltip", "Replace: Applies the current value minus the value to the new value."))
+				.ToolTip(LOCTEXT("BrushReplaceTooltip", "Replace: Overwrites the existing vertex value with the brush value."))
 				+ SSegmentedControl<EMeshVertexAttributePaintToolEditOperation>::Slot(EMeshVertexAttributePaintToolEditOperation::Multiply)
 				.Text(LOCTEXT("BrushMultiplyMode", "Multiply"))
-				.ToolTip(LOCTEXT("BrushMultiplyTooltip", "Multiply: Applies the current value multiplied by the value to the new value."))
+				.ToolTip(LOCTEXT("BrushMultiplyTooltip", "Multiply: Scales the existing vertex value by the brush value."))
 				+ SSegmentedControl<EMeshVertexAttributePaintToolEditOperation>::Slot(EMeshVertexAttributePaintToolEditOperation::Relax)
 				.Text(LOCTEXT("BrushRelaxMode", "Relax"))
-				.ToolTip(LOCTEXT("BrushRelaxTooltip", "Relax: Applies the average of the connected (by edge) vertex value to the new vertex value."))
+				.ToolTip(LOCTEXT("BrushRelaxTooltip", "Relax: Smooth the vertex value towards the average of its edge-connected neighbors."))
 			]
 		];
 	}
